Use the developer tag file when adding a new tag to defaultgameplaytags.ini when the developer setting is enabled
Duplication of CL #3185990

diff --git a/Engine/Plugins/Editor/GameplayTagsEditor/Source/GameplayTagsEditor/Private/GameplayTagsEditorModule.cpp b/Engine/Plugins/Editor/GameplayTagsEditor/Source/GameplayTagsEditor/Private/GameplayTagsEditorModule.cpp
--- a/Engine/Plugins/Editor/GameplayTagsEditor/Source/GameplayTagsEditor/Private/GameplayTagsEditorModule.cpp
+++ b/Engine/Plugins/Editor/GameplayTagsEditor/Source/GameplayTagsEditor/Private/GameplayTagsEditorModule.cpp
@@ -177,6 +177,8 @@ public:
 		// Write out gameplaytags.ini
 		GameplayTagsUpdateSourceControl(Settings->GetDefaultConfigFilename());
 		Settings->UpdateDefaultConfigFile();
+
+		GConfig->LoadFile(Settings->GetDefaultConfigFilename());
 		
 		// Write out all other tag lists
 		TArray<const FGameplayTagSource*> Sources;
@@ -252,11 +254,12 @@ public:
 			return false;
 		}
 
-		if (TagSourceName == NAME_None && DevSettings && DevSettings->DeveloperConfigName.IsEmpty() == false)
+		if ((TagSourceName == NAME_None || TagSourceName == FGameplayTagSource::GetDefaultName()) && DevSettings && !DevSettings->DeveloperConfigName.IsEmpty())
 		{
 			// Try to use developer config file
 			TagSourceName = FName(*FString::Printf(TEXT("%s.ini"), *DevSettings->DeveloperConfigName));
 		}
+
 		if (TagSourceName == NAME_None)
 		{
 			// If not set yet, set to default
@@ -283,6 +286,7 @@ public:
 
 			TagList->UpdateDefaultConfigFile(TagList->ConfigFileName);
 			GameplayTagsUpdateSourceControl(TagList->ConfigFileName);
+			GConfig->LoadFile(TagList->ConfigFileName);
 		}
 		else
 		{
@@ -322,6 +326,7 @@ public:
 
 					GameplayTagsUpdateSourceControl(Settings->GetDefaultConfigFilename());
 					Settings->UpdateDefaultConfigFile();
+					GConfig->LoadFile(Settings->GetDefaultConfigFilename());
 
 					ShowNotification(FText::Format(LOCTEXT("RemoveTagRedirect", "Deleted tag redirect {0}"), FText::FromString(TagToDelete)), 3.0f);
 
@@ -369,6 +374,7 @@ public:
 
 				TagList->UpdateDefaultConfigFile(TagList->ConfigFileName);
 				GameplayTagsUpdateSourceControl(TagList->ConfigFileName);
+				GConfig->LoadFile(TagList->ConfigFileName);
 
 				ShowNotification(FText::Format(LOCTEXT("RemoveTag", "Deleted tag {0}"), FText::FromString(TagToDelete)), 3.0f);
 
@@ -422,6 +428,7 @@ public:
 
 						TagList->UpdateDefaultConfigFile(TagList->ConfigFileName);
 						GameplayTagsUpdateSourceControl(TagList->ConfigFileName);
+						GConfig->LoadFile(TagList->ConfigFileName);
 
 						break;
 					}
@@ -438,6 +445,7 @@ public:
 
 		GameplayTagsUpdateSourceControl(Settings->GetDefaultConfigFilename());
 		Settings->UpdateDefaultConfigFile();
+		GConfig->LoadFile(Settings->GetDefaultConfigFilename());
 
 		ShowNotification(FText::Format(LOCTEXT("AddTagRedirect", "Renamed tag {0} to {1}"), FText::FromString(TagToRename), FText::FromString(TagToRename)), 3.0f);
 
