GameplayTag fixes from //Orion/Main:
-Remove tag log spam on startup
-Fix crash in editor when saving GameplayTag Query on blueprint CDO

diff --git a/Engine/Plugins/Editor/GameplayTagsEditor/Source/GameplayTagsEditor/Private/GameplayTagQueryCustomization.cpp b/Engine/Plugins/Editor/GameplayTagsEditor/Source/GameplayTagsEditor/Private/GameplayTagQueryCustomization.cpp
--- a/Engine/Plugins/Editor/GameplayTagsEditor/Source/GameplayTagsEditor/Private/GameplayTagQueryCustomization.cpp
+++ b/Engine/Plugins/Editor/GameplayTagsEditor/Source/GameplayTagsEditor/Private/GameplayTagQueryCustomization.cpp
@@ -241,8 +241,11 @@ void FGameplayTagQueryCustomization::CloseWidgetWindow(bool WasCancelled)
 	// Notify change. This is required for these to work inside of UDataTables
 	if (!WasCancelled && PropertyUtilities.IsValid())
 	{
+		UProperty* TheProperty = StructPropertyHandle->GetProperty();
+
 		FEditPropertyChain PropertyChain;
-		PropertyChain.SetActivePropertyNode(StructPropertyHandle->GetProperty());
+		PropertyChain.AddHead(TheProperty);
+		PropertyChain.SetActivePropertyNode(TheProperty);
 
 		FPropertyChangedEvent ChangeEvent(StructPropertyHandle->GetProperty(), EPropertyChangeType::ValueSet, nullptr);
 		FNotifyHook* NotifyHook = PropertyUtilities->GetNotifyHook();
diff --git a/Engine/Source/Runtime/GameplayTags/Private/GameplayTagsManager.cpp b/Engine/Source/Runtime/GameplayTags/Private/GameplayTagsManager.cpp
--- a/Engine/Source/Runtime/GameplayTags/Private/GameplayTagsManager.cpp
+++ b/Engine/Source/Runtime/GameplayTags/Private/GameplayTagsManager.cpp
@@ -1175,15 +1175,6 @@ void UGameplayTagsManager::DoneAddingNativeTags()
 		{
 			ConstructNetIndex();
 		}
-
-		// TEMP
-		for (auto It : GameplayTagNodeMap)
-		{
-			FGameplayTag Tag = It.Key;
-			TSharedPtr<FGameplayTagNode> Node = It.Value;
-
-			UE_LOG(LogGameplayTags, Display, TEXT("Tag %s NetIndex: %d"), *Tag.ToString(), Node->GetNetIndex());
-		}
 	}
 }
 
