Mirrors CL 4133776

Fix a bug where the incorrect error message was displayed when trying to delete implicit tags.

#rnx

diff --git a/Engine/Plugins/Editor/GameplayTagsEditor/Source/GameplayTagsEditor/Private/GameplayTagsEditorModule.cpp b/Engine/Plugins/Editor/GameplayTagsEditor/Source/GameplayTagsEditor/Private/GameplayTagsEditorModule.cpp
--- a/Engine/Plugins/Editor/GameplayTagsEditor/Source/GameplayTagsEditor/Private/GameplayTagsEditorModule.cpp
+++ b/Engine/Plugins/Editor/GameplayTagsEditor/Source/GameplayTagsEditor/Private/GameplayTagsEditorModule.cpp
@@ -495,8 +495,8 @@ public:
 
 		const FGameplayTagSource* TagSource = Manager.FindTagSource(TagSourceName);
 
-		// Verify tag source
-		if (!TagSource)
+		// Check if the tag is implicitly defined
+		if (!bTagIsExplicit || !TagSource)
 		{
 			ShowNotification(FText::Format(LOCTEXT("RemoveTagFailureNoSource", "Cannot delete tag {0} as it is implicit, remove children manually"), FText::FromString(TagToDelete)), 10.0f);
 
