Mirrors CL 4133778

Remove the option to rename or delete tags for tags that were defined in datatables.

#rnx

diff --git a/Engine/Plugins/Editor/GameplayTagsEditor/Source/GameplayTagsEditor/Private/SGameplayTagWidget.cpp b/Engine/Plugins/Editor/GameplayTagsEditor/Source/GameplayTagsEditor/Private/SGameplayTagWidget.cpp
--- a/Engine/Plugins/Editor/GameplayTagsEditor/Source/GameplayTagsEditor/Private/SGameplayTagWidget.cpp
+++ b/Engine/Plugins/Editor/GameplayTagsEditor/Source/GameplayTagsEditor/Private/SGameplayTagWidget.cpp
@@ -843,7 +843,13 @@ TSharedRef<SWidget> SGameplayTagWidget::MakeTagActionsMenu(TSharedPtr<FGameplayT
 		bShowManagement = false;
 	}
 
-	FMenuBuilder MenuBuilder(true, NULL);
+	// we can only rename or delete tags if they came from an ini file
+	if (!InTagNode->SourceName.ToString().EndsWith(TEXT(".ini")))
+	{
+		bShowManagement = false;
+	}
+
+	FMenuBuilder MenuBuilder(true, nullptr);
 
 	// Rename
 	if (bShowManagement)
@@ -1079,7 +1085,7 @@ void SGameplayTagWidget::SetContainer(FGameplayTagContainer* OriginalContainer,
 		// Not sure if we should get here, means the property handle hasnt been setup which could be right or wrong.
 		if (OwnerObj)
 		{
-			OwnerObj->PreEditChange(PropertyHandle.IsValid() ? PropertyHandle->GetProperty() : NULL);
+			OwnerObj->PreEditChange(PropertyHandle.IsValid() ? PropertyHandle->GetProperty() : nullptr);
 		}
 
 		*OriginalContainer = *EditedContainer;
