Fix a few minor issues with restricted gameplay tag editing.
 - Refreshing the tag list no longer puts unrestricted tags in the restricted tag display
 - Fixed a bug where the add subtags button on restricted tags didn't work correctly.

#rnx

diff --git a/Engine/Plugins/Editor/GameplayTagsEditor/Source/GameplayTagsEditor/Private/SGameplayTagWidget.cpp b/Engine/Plugins/Editor/GameplayTagsEditor/Source/GameplayTagsEditor/Private/SGameplayTagWidget.cpp
--- a/Engine/Plugins/Editor/GameplayTagsEditor/Source/GameplayTagsEditor/Private/SGameplayTagWidget.cpp
+++ b/Engine/Plugins/Editor/GameplayTagsEditor/Source/GameplayTagsEditor/Private/SGameplayTagWidget.cpp
@@ -818,11 +818,16 @@ FReply SGameplayTagWidget::OnAddSubtagClicked(TSharedPtr<FGameplayTagNode> InTag
 
 		Manager.GetTagEditorData(InTagNode->GetCompleteTagName(), TagComment, TagSource, bTagIsExplicit, bTagIsRestricted, bTagAllowsNonRestrictedChildren);
 
-		if (AddNewTagWidget.IsValid())
+		if (!bRestrictedTags && AddNewTagWidget.IsValid())
 		{
 			bAddTagSectionExpanded = true; 
 			AddNewTagWidget->AddSubtagFromParent(TagName, TagSource);
 		}
+		else if (bRestrictedTags && AddNewRestrictedTagWidget.IsValid())
+		{
+			bAddTagSectionExpanded = true;
+			AddNewRestrictedTagWidget->AddSubtagFromParent(TagName, TagSource, bTagAllowsNonRestrictedChildren);
+		}
 	}
 	return FReply::Handled();
 }
@@ -1130,6 +1135,18 @@ void SGameplayTagWidget::RefreshTags()
 	UGameplayTagsManager& Manager = UGameplayTagsManager::Get();
 	Manager.GetFilteredGameplayRootTags(RootFilterString, TagItems);
 
+	if (bRestrictedTags)
+	{
+		// We only want to show the restricted gameplay tags
+		for (int32 Idx = TagItems.Num() - 1; Idx >= 0; --Idx)
+		{
+			if (!TagItems[Idx]->IsRestrictedGameplayTag())
+			{
+				TagItems.RemoveAtSwap(Idx);
+			}
+		}
+	}
+
 	TagTreeWidget->SetTreeItemsSource(&TagItems);
 }
 
