Fixing misc restricted gameplay tag UI bugs

diff --git a/Engine/Plugins/Editor/GameplayTagsEditor/Source/GameplayTagsEditor/Private/GameplayTagsEditorModule.cpp b/Engine/Plugins/Editor/GameplayTagsEditor/Source/GameplayTagsEditor/Private/GameplayTagsEditorModule.cpp
--- a/Engine/Plugins/Editor/GameplayTagsEditor/Source/GameplayTagsEditor/Private/GameplayTagsEditorModule.cpp
+++ b/Engine/Plugins/Editor/GameplayTagsEditor/Source/GameplayTagsEditor/Private/GameplayTagsEditorModule.cpp
@@ -280,16 +280,14 @@ public:
 		}
 	}
 
-	bool DeleteTagRedirector(const FString& TagToDelete)
+	bool DeleteTagRedirector(const FName& TagToDelete)
 	{
-		FName TagName = FName(*TagToDelete);
-
 		UGameplayTagsSettings* Settings = GetMutableDefault<UGameplayTagsSettings>();
 		UGameplayTagsManager& Manager = UGameplayTagsManager::Get();
 
 		for (int32 i = 0; i < Settings->GameplayTagRedirects.Num(); i++)
 		{
-			if (Settings->GameplayTagRedirects[i].OldTagName == TagName)
+			if (Settings->GameplayTagRedirects[i].OldTagName == TagToDelete)
 			{
 				Settings->GameplayTagRedirects.RemoveAt(i);
 
@@ -299,7 +297,7 @@ public:
 
 				Manager.EditorRefreshGameplayTagTree();
 
-				ShowNotification(FText::Format(LOCTEXT("RemoveTagRedirect", "Deleted tag redirect {0}"), FText::FromString(TagToDelete)), 5.0f);
+				ShowNotification(FText::Format(LOCTEXT("RemoveTagRedirect", "Deleted tag redirect {0}"), FText::FromName(TagToDelete)), 5.0f);
 
 				return true;
 			}
@@ -325,11 +323,13 @@ public:
 		UGameplayTagsSettings*			Settings = GetMutableDefault<UGameplayTagsSettings>();
 		UGameplayTagsDeveloperSettings* DevSettings = GetMutableDefault<UGameplayTagsDeveloperSettings>();
 
+		FName NewTagName = FName(*NewTag);
+
 		// Delete existing redirector
-		DeleteTagRedirector(NewTag);
+		DeleteTagRedirector(NewTagName);
 
 		// Already in the list as an explicit tag, ignore. Note we want to add if it is in implicit tag. (E.g, someone added A.B.C then someone tries to add A.B)
-		if (Manager.IsDictionaryTag(FName(*NewTag)))
+		if (Manager.IsDictionaryTag(NewTagName))
 		{
 			ShowNotification(FText::Format(LOCTEXT("AddTagFailure_AlreadyExists", "Failed to add gameplay tag {0}, already exists!"), FText::FromString(NewTag)), 10.0f);
 
@@ -468,9 +468,9 @@ public:
 		return true;
 	}
 
-	virtual bool DeleteTagFromINI(const FString& TagToDelete) override
+	virtual bool DeleteTagFromINI(TSharedPtr<FGameplayTagNode> TagNodeToDelete) override
 	{
-		FName TagName = FName(*TagToDelete);
+		FName TagName = TagNodeToDelete->GetCompleteTagName();
 
 		UGameplayTagsManager& Manager = UGameplayTagsManager::Get();
 		UGameplayTagsSettings* Settings = GetMutableDefault<UGameplayTagsSettings>();
@@ -481,31 +481,37 @@ public:
 		bool bTagIsRestricted;
 		bool bTagAllowsNonRestrictedChildren;
 
-		if (DeleteTagRedirector(TagToDelete))
+		if (DeleteTagRedirector(TagName))
 		{
 			return true;
 		}
 		
 		if (!Manager.GetTagEditorData(TagName, Comment, TagSourceName, bTagIsExplicit, bTagIsRestricted, bTagAllowsNonRestrictedChildren))
 		{
-			ShowNotification(FText::Format(LOCTEXT("RemoveTagFailureNoTag", "Cannot delete tag {0}, does not exist!"), FText::FromString(TagToDelete)), 10.0f);
-
+			ShowNotification(FText::Format(LOCTEXT("RemoveTagFailureNoTag", "Cannot delete tag {0}, does not exist!"), FText::FromName(TagName)), 10.0f);
 			return false;
 		}
 
+		ensure(bTagIsRestricted == TagNodeToDelete->IsRestrictedGameplayTag());
+
 		const FGameplayTagSource* TagSource = Manager.FindTagSource(TagSourceName);
 
 		// Check if the tag is implicitly defined
 		if (!bTagIsExplicit || !TagSource)
 		{
-			ShowNotification(FText::Format(LOCTEXT("RemoveTagFailureNoSource", "Cannot delete tag {0} as it is implicit, remove children manually"), FText::FromString(TagToDelete)), 10.0f);
-
+			ShowNotification(FText::Format(LOCTEXT("RemoveTagFailureNoSource", "Cannot delete tag {0} as it is implicit, remove children manually"), FText::FromName(TagName)), 10.0f);
 			return false;
 		}
-		else if (!TagSource->SourceTagList)
+		
+		if (bTagIsRestricted && !TagSource->SourceRestrictedTagList)
 		{
-			ShowNotification(FText::Format(LOCTEXT("RemoveTagFailureBadSource", "Cannot delete tag {0} from source {1}, remove manually"), FText::FromString(TagToDelete), FText::FromName(TagSourceName)), 10.0f);
+			ShowNotification(FText::Format(LOCTEXT("RemoveTagFailureBadSource", "Cannot delete tag {0} from source {1}, remove manually"), FText::FromName(TagName), FText::FromName(TagSourceName)), 10.0f);
+			return false;
+		}
 
+		if (!bTagIsRestricted && !TagSource->SourceTagList)
+		{
+			ShowNotification(FText::Format(LOCTEXT("RemoveTagFailureBadSource", "Cannot delete tag {0} from source {1}, remove manually"), FText::FromName(TagName), FText::FromName(TagSourceName)), 10.0f);
 			return false;
 		}
 
@@ -517,7 +523,7 @@ public:
 		TagsThatWillBeDeleted.Add(TagName);
 
 		FGameplayTag ParentTag = ActualTag.RequestDirectParent();
-		while (ParentTag.IsValid())
+		while (ParentTag.IsValid() && !Manager.FindTagNode(ParentTag)->IsExplicitTag())
 		{
 			// See if there are more children than the one we are about to delete
 			FGameplayTagContainer ParentChildTags = Manager.RequestGameplayTagChildrenInDictionary(ParentTag);
@@ -553,27 +559,45 @@ public:
 		}
 
 		// Passed, delete and save
-		UGameplayTagsList* TagList = TagSource->SourceTagList;
+		const FString& ConfigFileName = bTagIsRestricted ? TagSource->SourceRestrictedTagList->ConfigFileName : TagSource->SourceTagList->ConfigFileName;
+		int32 TagListSize = bTagIsRestricted ? TagSource->SourceRestrictedTagList->RestrictedGameplayTagList.Num() : TagSource->SourceTagList->GameplayTagList.Num();
 
-		for (int32 i = 0; i < TagList->GameplayTagList.Num(); i++)
+		for (int32 i = 0; i < TagListSize; i++)
 		{
-			if (TagList->GameplayTagList[i].Tag == TagName)
+			bool bRemoved = false;
+			if (bTagIsRestricted)
 			{
-				TagList->GameplayTagList.RemoveAt(i);
+				if (TagSource->SourceRestrictedTagList->RestrictedGameplayTagList[i].Tag == TagName)
+				{
+					TagSource->SourceRestrictedTagList->RestrictedGameplayTagList.RemoveAt(i);
+					TagSource->SourceRestrictedTagList->UpdateDefaultConfigFile(ConfigFileName);
+					bRemoved = true;
+				}
+			}
+			else
+			{
+				if (TagSource->SourceTagList->GameplayTagList[i].Tag == TagName)
+				{
+					TagSource->SourceTagList->GameplayTagList.RemoveAt(i);
+					TagSource->SourceTagList->UpdateDefaultConfigFile(ConfigFileName);
+					bRemoved = true;
+				}
+			}
 
-				TagList->UpdateDefaultConfigFile(TagList->ConfigFileName);
-				GameplayTagsUpdateSourceControl(TagList->ConfigFileName);
-				GConfig->LoadFile(TagList->ConfigFileName);
+			if (bRemoved)
+			{
+				GameplayTagsUpdateSourceControl(ConfigFileName);
+				GConfig->LoadFile(ConfigFileName);
 
 				// See if we still live due to child tags
 
 				if (ChildTags.Num() > 0)
 				{
-					ShowNotification(FText::Format(LOCTEXT("RemoveTagChildrenExist", "Deleted explicit tag {0}, still exists implicitly due to children"), FText::FromString(TagToDelete)), 5.0f);
+					ShowNotification(FText::Format(LOCTEXT("RemoveTagChildrenExist", "Deleted explicit tag {0}, still exists implicitly due to children"), FText::FromName(TagName)), 5.0f);
 				}
 				else
 				{
-					ShowNotification(FText::Format(LOCTEXT("RemoveTag", "Deleted tag {0}"), FText::FromString(TagToDelete)), 5.0f);
+					ShowNotification(FText::Format(LOCTEXT("RemoveTag", "Deleted tag {0}"), FText::FromName(TagName)), 5.0f);
 				}
 
 				// This invalidates all local variables, need to return right away
@@ -583,11 +607,95 @@ public:
 			}
 		}
 
-		ShowNotification(FText::Format(LOCTEXT("RemoveTagFailureNoTag", "Cannot delete tag {0}, does not exist!"), FText::FromString(TagToDelete)), 10.0f);
+		ShowNotification(FText::Format(LOCTEXT("RemoveTagFailureNoTag", "Cannot delete tag {0}, does not exist!"), FText::FromName(TagName)), 10.0f);
 		
 		return false;
 	}
 
+	virtual bool UpdateTagInINI(const FString& TagToUpdate, const FString& Comment, bool bIsRestrictedTag, bool bAllowNonRestrictedChildren) override
+	{
+		FName TagName = FName(*TagToUpdate);
+
+		UGameplayTagsManager& Manager = UGameplayTagsManager::Get();
+		UGameplayTagsSettings* Settings = GetMutableDefault<UGameplayTagsSettings>();
+
+		FString OldComment;
+		FName TagSourceName;
+		bool bTagIsExplicit;
+		bool bTagWasRestricted;
+		bool bTagDidAllowNonRestrictedChildren;
+
+		bool bSuccess = false;
+
+		if (Manager.GetTagEditorData(TagName, OldComment, TagSourceName, bTagIsExplicit, bTagWasRestricted, bTagDidAllowNonRestrictedChildren))
+		{
+			if (const FGameplayTagSource* TagSource = Manager.FindTagSource(TagSourceName))
+			{
+				// if we're disallowing non-restricted children make sure we don't already have some
+				if (bTagDidAllowNonRestrictedChildren && !bAllowNonRestrictedChildren)
+				{
+					FGameplayTag ActualTag = Manager.RequestGameplayTag(TagName);
+					FGameplayTagContainer ChildTags = Manager.RequestGameplayTagDirectDescendantsInDictionary(ActualTag, EGameplayTagSelectionType::RestrictedOnly);
+					if (!ChildTags.IsEmpty())
+					{
+						ShowNotification(LOCTEXT("ToggleAllowNonRestrictedChildrenFailure", "Cannot prevent non-restricted children since some already exist! Delete them first."), 10.0f);
+						return false;
+					}
+				}
+
+				UObject* TagListObj = nullptr;
+				FString ConfigFileName;
+
+				if (bIsRestrictedTag && TagSource->SourceRestrictedTagList)
+				{
+					URestrictedGameplayTagsList* RestrictedTagList = TagSource->SourceRestrictedTagList;
+					TagListObj = RestrictedTagList;
+					ConfigFileName = RestrictedTagList->ConfigFileName;
+
+					for (int32 i = 0; i < RestrictedTagList->RestrictedGameplayTagList.Num(); i++)
+					{
+						if (RestrictedTagList->RestrictedGameplayTagList[i].Tag == TagName)
+						{
+							RestrictedTagList->RestrictedGameplayTagList[i].bAllowNonRestrictedChildren = bAllowNonRestrictedChildren;
+							bSuccess = true;
+							break;
+						}
+					}
+				}
+
+				// nothing to do here for non-restricted tags but this is the code that would handle it if we expand this function
+				// 				else if (TagSource->SourceTagList)
+				// 				{
+				// 					UGameplayTagsList* TagList = TagSource->SourceTagList;
+				//					TagListObj = TagList;
+				// 
+				// 					for (int32 i = 0; i < TagList->GameplayTagList.Num(); i++)
+				// 					{
+				// 						if (TagList->GameplayTagList[i].Tag == TagName)
+				// 						{
+				// 
+				// 							bSuccess = true;
+				// 							break;
+				// 						}
+				// 					}
+				// 				}
+
+				if (bSuccess)
+				{
+					// Check source control before and after writing, to make sure it gets created or checked out
+					GameplayTagsUpdateSourceControl(ConfigFileName);
+					TagListObj->UpdateDefaultConfigFile(ConfigFileName);
+					GameplayTagsUpdateSourceControl(ConfigFileName);
+
+					GConfig->LoadFile(ConfigFileName);
+				}
+
+			}
+		}
+
+		return bSuccess;
+	}
+
 	virtual bool RenameTagInINI(const FString& TagToRename, const FString& TagToRenameTo) override
 	{
 		FName OldTagName = FName(*TagToRename);
@@ -603,8 +711,8 @@ public:
 		bool bTagAllowsNonRestrictedChildren;
 
 		// Delete existing redirector
-		DeleteTagRedirector(TagToRenameTo);
-		DeleteTagRedirector(TagToRename);
+		DeleteTagRedirector(NewTagName);
+		DeleteTagRedirector(OldTagName);
 
 		if (Manager.GetTagEditorData(OldTagName, OldComment, OldTagSourceName, bTagIsExplicit, bTagIsRestricted, bTagAllowsNonRestrictedChildren))
 		{
diff --git a/Engine/Plugins/Editor/GameplayTagsEditor/Source/GameplayTagsEditor/Private/SGameplayTagWidget.cpp b/Engine/Plugins/Editor/GameplayTagsEditor/Source/GameplayTagsEditor/Private/SGameplayTagWidget.cpp
--- a/Engine/Plugins/Editor/GameplayTagsEditor/Source/GameplayTagsEditor/Private/SGameplayTagWidget.cpp
+++ b/Engine/Plugins/Editor/GameplayTagsEditor/Source/GameplayTagsEditor/Private/SGameplayTagWidget.cpp
@@ -712,13 +712,18 @@ ECheckBoxState SGameplayTagWidget::IsTagChecked(TSharedPtr<FGameplayTagNode> Nod
 
 void SGameplayTagWidget::OnAllowChildrenTagCheckStatusChanged(ECheckBoxState NewCheckState, TSharedPtr<FGameplayTagNode> NodeChanged)
 {
-	if (NewCheckState == ECheckBoxState::Checked)
-	{
-		NodeChanged->bAllowNonRestrictedChildren = true;
-	}
-	else if (NewCheckState == ECheckBoxState::Unchecked)
+	IGameplayTagsEditorModule& TagsEditor = IGameplayTagsEditorModule::Get();
+
+	if (TagsEditor.UpdateTagInINI(NodeChanged->GetCompleteTagString(), NodeChanged->DevComment, NodeChanged->bIsRestrictedTag, NewCheckState == ECheckBoxState::Checked))
 	{
-		NodeChanged->bAllowNonRestrictedChildren = false;
+		if (NewCheckState == ECheckBoxState::Checked)
+		{
+			NodeChanged->bAllowNonRestrictedChildren = true;
+		}
+		else if (NewCheckState == ECheckBoxState::Unchecked)
+		{
+			NodeChanged->bAllowNonRestrictedChildren = false;
+		}
 	}
 }
 
@@ -898,7 +903,7 @@ void SGameplayTagWidget::OnDeleteTag(TSharedPtr<FGameplayTagNode> InTagNode)
 	{
 		IGameplayTagsEditorModule& TagsEditor = IGameplayTagsEditorModule::Get();
 
-		const bool bDeleted = TagsEditor.DeleteTagFromINI(InTagNode->GetCompleteTagString());
+		const bool bDeleted = TagsEditor.DeleteTagFromINI(InTagNode);
 
 		if (bDeleted)
 		{
diff --git a/Engine/Plugins/Editor/GameplayTagsEditor/Source/GameplayTagsEditor/Public/GameplayTagsEditorModule.h b/Engine/Plugins/Editor/GameplayTagsEditor/Source/GameplayTagsEditor/Public/GameplayTagsEditorModule.h
--- a/Engine/Plugins/Editor/GameplayTagsEditor/Source/GameplayTagsEditor/Public/GameplayTagsEditorModule.h
+++ b/Engine/Plugins/Editor/GameplayTagsEditor/Source/GameplayTagsEditor/Public/GameplayTagsEditorModule.h
@@ -41,11 +41,14 @@ public:
 	GAMEPLAYTAGSEDITOR_API virtual bool AddNewGameplayTagToINI(const FString& NewTag, const FString& Comment = TEXT(""), FName TagSourceName = NAME_None, bool bIsRestrictedTag = false, bool bAllowNonRestrictedChildren = true) = 0;
 
 	/** Tries to delete a tag from the library. This will pop up special UI or error messages as needed. It will also delete redirectors if that is specified. */
-	GAMEPLAYTAGSEDITOR_API virtual bool DeleteTagFromINI(const FString& TagToDelete) = 0;
+	GAMEPLAYTAGSEDITOR_API virtual bool DeleteTagFromINI(TSharedPtr<struct FGameplayTagNode> TagNodeToDelete) = 0;
 
 	/** Tries to rename a tag, leaving a rediretor in the ini, and adding the new tag if it does not exist yet */
 	GAMEPLAYTAGSEDITOR_API virtual bool RenameTagInINI(const FString& TagToRename, const FString& TagToRenameTo) = 0;
 
+	/** Updates info about a tag */
+	GAMEPLAYTAGSEDITOR_API virtual bool UpdateTagInINI(const FString& TagToUpdate, const FString& Comment, bool bIsRestrictedTag, bool bAllowNonRestrictedChildren) = 0;
+
 	/** Adds a transient gameplay tag (only valid for the current editor session) */
 	GAMEPLAYTAGSEDITOR_API virtual bool AddTransientEditorGameplayTag(const FString& NewTransientTag) = 0;
 };
diff --git a/Engine/Source/Runtime/GameplayTags/Classes/GameplayTagsManager.h b/Engine/Source/Runtime/GameplayTags/Classes/GameplayTagsManager.h
--- a/Engine/Source/Runtime/GameplayTags/Classes/GameplayTagsManager.h
+++ b/Engine/Source/Runtime/GameplayTags/Classes/GameplayTagsManager.h
@@ -72,6 +72,15 @@ enum class EGameplayTagSourceType : uint8
 	Invalid,			// Not a real source
 };
 
+UENUM()
+enum class EGameplayTagSelectionType : uint8
+{
+	None,
+	NonRestrictedOnly,
+	RestrictedOnly,
+	All
+};
+
 /** Struct defining where gameplay tags are loaded/saved from. Mostly for the editor */
 USTRUCT()
 struct FGameplayTagSource
@@ -511,8 +520,12 @@ class GAMEPLAYTAGS_API UGameplayTagsManager : public UObject
 	/** Refresh the gameplaytag tree due to an editor change */
 	void EditorRefreshGameplayTagTree();
 
-	/** Gets a Tag Container containing the all tags in the hierarchy that are children of this tag, and were explicitly added to the dictionary */
+	/** Gets a Tag Container containing all of the tags in the hierarchy that are children of this tag, and were explicitly added to the dictionary */
 	FGameplayTagContainer RequestGameplayTagChildrenInDictionary(const FGameplayTag& GameplayTag) const;
+#if WITH_EDITORONLY_DATA
+	/** Gets a Tag Container containing all of the tags in the hierarchy that are children of this tag, were explicitly added to the dictionary, and do not have any explicitly added tags between them and the specified tag */
+	FGameplayTagContainer RequestGameplayTagDirectDescendantsInDictionary(const FGameplayTag& GameplayTag, EGameplayTagSelectionType SelectionType) const;
+#endif // WITH_EDITORONLY_DATA
 
 	/** This is called when EditorRefreshGameplayTagTree. Useful if you need to do anything editor related when tags are added or removed */
 	static FSimpleMulticastDelegate OnEditorRefreshGameplayTagTree;
diff --git a/Engine/Source/Runtime/GameplayTags/Private/GameplayTagsManager.cpp b/Engine/Source/Runtime/GameplayTags/Private/GameplayTagsManager.cpp
--- a/Engine/Source/Runtime/GameplayTags/Private/GameplayTagsManager.cpp
+++ b/Engine/Source/Runtime/GameplayTags/Private/GameplayTagsManager.cpp
@@ -1239,6 +1239,48 @@ FGameplayTagContainer UGameplayTagsManager::RequestGameplayTagChildrenInDictiona
 	return TagContainer;
 }
 
+#if WITH_EDITORONLY_DATA
+FGameplayTagContainer UGameplayTagsManager::RequestGameplayTagDirectDescendantsInDictionary(const FGameplayTag& GameplayTag, EGameplayTagSelectionType SelectionType) const
+{
+	bool bIncludeRestrictedTags = (SelectionType == EGameplayTagSelectionType::RestrictedOnly || SelectionType == EGameplayTagSelectionType::All);
+	bool bIncludeNonRestrictedTags = (SelectionType == EGameplayTagSelectionType::NonRestrictedOnly || SelectionType == EGameplayTagSelectionType::All);
+
+	// Note this purposefully does not include the passed in GameplayTag in the container.
+	FGameplayTagContainer TagContainer;
+
+	TSharedPtr<FGameplayTagNode> GameplayTagNode = FindTagNode(GameplayTag);
+	if (GameplayTagNode.IsValid())
+	{
+		TArray< TSharedPtr<FGameplayTagNode> >& ChildrenNodes = GameplayTagNode->GetChildTagNodes();
+		int32 CurrArraySize = ChildrenNodes.Num();
+		for (int32 Idx = 0; Idx < CurrArraySize; ++Idx)
+		{
+			TSharedPtr<FGameplayTagNode> ChildNode = ChildrenNodes[Idx];
+			if (ChildNode.IsValid())
+			{
+				// if the tag isn't in the dictionary, add its children to the list
+				if (ChildNode->SourceName == NAME_None)
+				{
+					TArray< TSharedPtr<FGameplayTagNode> >& GrandChildrenNodes = ChildNode->GetChildTagNodes();
+					ChildrenNodes.Append(GrandChildrenNodes);
+					CurrArraySize = ChildrenNodes.Num();
+				}
+				else
+				{
+					// this tag is in the dictionary so add it to the list
+					if ((ChildNode->bIsRestrictedTag && bIncludeRestrictedTags) ||
+						(!ChildNode->bIsRestrictedTag && bIncludeNonRestrictedTags))
+					{
+						TagContainer.AddTag(ChildNode->GetCompleteTag());
+					}
+				}
+			}
+		}
+	}
+	return TagContainer;
+}
+#endif // WITH_EDITORONLY_DATA
+
 void UGameplayTagsManager::NotifyGameplayTagDoubleClickedEditor(FString TagName)
 {
 	FGameplayTag Tag = RequestGameplayTag(FName(*TagName), false);
