Added an #include to fix the non-unity build.

diff --git a/Engine/Plugins/Editor/GameplayTagsEditor/Source/GameplayTagsEditor/Private/SAddNewGameplayTagWidget.cpp b/Engine/Plugins/Editor/GameplayTagsEditor/Source/GameplayTagsEditor/Private/SAddNewGameplayTagWidget.cpp
--- a/Engine/Plugins/Editor/GameplayTagsEditor/Source/GameplayTagsEditor/Private/SAddNewGameplayTagWidget.cpp
+++ b/Engine/Plugins/Editor/GameplayTagsEditor/Source/GameplayTagsEditor/Private/SAddNewGameplayTagWidget.cpp
@@ -6,6 +6,7 @@
 #include "GameplayTagsEditorModule.h"
 #include "GameplayTagsModule.h"
 #include "Widgets/Input/SButton.h"
+#include "Misc/MessageDialog.h"
 
 #define LOCTEXT_NAMESPACE "AddNewGameplayTagWidget"
 
