Fix for Mac editor linking

diff --git a/Engine/Plugins/Editor/GameplayTagsEditor/Source/GameplayTagsEditor/Private/SAddNewGameplayTagWidget.h b/Engine/Plugins/Editor/GameplayTagsEditor/Source/GameplayTagsEditor/Private/SAddNewGameplayTagWidget.h
--- a/Engine/Plugins/Editor/GameplayTagsEditor/Source/GameplayTagsEditor/Private/SAddNewGameplayTagWidget.h
+++ b/Engine/Plugins/Editor/GameplayTagsEditor/Source/GameplayTagsEditor/Private/SAddNewGameplayTagWidget.h
@@ -12,7 +12,7 @@
 #include "Widgets/Input/SComboBox.h"
 
 /** Widget allowing the user to create new gameplay tags */
-class SAddNewGameplayTagWidget : public SCompoundWidget
+class GAMEPLAYTAGSEDITOR_VTABLE SAddNewGameplayTagWidget : public SCompoundWidget
 {
 public:
 
