Fix a bug where you could set some restricted tags to not allow non-restricted children even though they already had non-restricted children.
Automatically add .ini file extension when adding restricted tag config files.
Don't display the warning message for modifying restricted tags if the current user is one of the tag owners.

diff --git a/Engine/Plugins/Editor/GameplayTagsEditor/Source/GameplayTagsEditor/Private/GameplayTagsEditorModule.cpp b/Engine/Plugins/Editor/GameplayTagsEditor/Source/GameplayTagsEditor/Private/GameplayTagsEditorModule.cpp
--- a/Engine/Plugins/Editor/GameplayTagsEditor/Source/GameplayTagsEditor/Private/GameplayTagsEditorModule.cpp
+++ b/Engine/Plugins/Editor/GameplayTagsEditor/Source/GameplayTagsEditor/Private/GameplayTagsEditorModule.cpp
@@ -637,7 +637,7 @@ public:
 				if (bTagDidAllowNonRestrictedChildren && !bAllowNonRestrictedChildren)
 				{
 					FGameplayTag ActualTag = Manager.RequestGameplayTag(TagName);
-					FGameplayTagContainer ChildTags = Manager.RequestGameplayTagDirectDescendantsInDictionary(ActualTag, EGameplayTagSelectionType::RestrictedOnly);
+					FGameplayTagContainer ChildTags = Manager.RequestGameplayTagDirectDescendantsInDictionary(ActualTag, EGameplayTagSelectionType::NonRestrictedOnly);
 					if (!ChildTags.IsEmpty())
 					{
 						ShowNotification(LOCTEXT("ToggleAllowNonRestrictedChildrenFailure", "Cannot prevent non-restricted children since some already exist! Delete them first."), 10.0f);
@@ -665,23 +665,6 @@ public:
 					}
 				}
 
-				// nothing to do here for non-restricted tags but this is the code that would handle it if we expand this function
-				// 				else if (TagSource->SourceTagList)
-				// 				{
-				// 					UGameplayTagsList* TagList = TagSource->SourceTagList;
-				//					TagListObj = TagList;
-				// 
-				// 					for (int32 i = 0; i < TagList->GameplayTagList.Num(); i++)
-				// 					{
-				// 						if (TagList->GameplayTagList[i].Tag == TagName)
-				// 						{
-				// 
-				// 							bSuccess = true;
-				// 							break;
-				// 						}
-				// 					}
-				// 				}
-
 				if (bSuccess)
 				{
 					// Check source control before and after writing, to make sure it gets created or checked out
diff --git a/Engine/Plugins/Editor/GameplayTagsEditor/Source/GameplayTagsEditor/Private/SAddNewRestrictedGameplayTagWidget.cpp b/Engine/Plugins/Editor/GameplayTagsEditor/Source/GameplayTagsEditor/Private/SAddNewRestrictedGameplayTagWidget.cpp
--- a/Engine/Plugins/Editor/GameplayTagsEditor/Source/GameplayTagsEditor/Private/SAddNewRestrictedGameplayTagWidget.cpp
+++ b/Engine/Plugins/Editor/GameplayTagsEditor/Source/GameplayTagsEditor/Private/SAddNewRestrictedGameplayTagWidget.cpp
@@ -281,23 +281,38 @@ void SAddNewRestrictedGameplayTagWidget::ValidateNewRestrictedTag()
 
 	if (bHasOwner)
 	{
-		FString StringToDisplay = TEXT("Do you have permission from ");
-		StringToDisplay.Append(TagSourceOwners[0]);
-		for (int Idx = 1; Idx < TagSourceOwners.Num(); ++Idx)
+		// check if we're one of the owners; if we are then we don't need to pop up the permission dialog
+		bool bRequiresPermission = true;
+		const FString& UserName = FPlatformProcess::UserName();
+		for (const FString& Owner : TagSourceOwners)
 		{
-			StringToDisplay.Append(TEXT(" or "));
-			StringToDisplay.Append(TagSourceOwners[Idx]);
+			if (Owner.Equals(UserName))
+			{
+				CreateNewRestrictedGameplayTag();
+				bRequiresPermission = false;
+			}
 		}
-		StringToDisplay.Append(TEXT(" to modify "));
-		StringToDisplay.Append(TagSource.ToString());
-		StringToDisplay.Append(TEXT("?"));
 
-		FNotificationInfo Info(FText::FromString(StringToDisplay));
-		Info.ExpireDuration = 10.f;
-		Info.ButtonDetails.Add(FNotificationButtonInfo(LOCTEXT("RestrictedTagPopupButtonAccept", "Yes"), FText(), FSimpleDelegate::CreateSP(this, &SAddNewRestrictedGameplayTagWidget::CreateNewRestrictedGameplayTag), SNotificationItem::CS_None));
-		Info.ButtonDetails.Add(FNotificationButtonInfo(LOCTEXT("RestrictedTagPopupButtonReject", "No"), FText(), FSimpleDelegate::CreateSP(this, &SAddNewRestrictedGameplayTagWidget::CancelNewTag), SNotificationItem::CS_None));
+		if (bRequiresPermission)
+		{
+			FString StringToDisplay = TEXT("Do you have permission from ");
+			StringToDisplay.Append(TagSourceOwners[0]);
+			for (int Idx = 1; Idx < TagSourceOwners.Num(); ++Idx)
+			{
+				StringToDisplay.Append(TEXT(" or "));
+				StringToDisplay.Append(TagSourceOwners[Idx]);
+			}
+			StringToDisplay.Append(TEXT(" to modify "));
+			StringToDisplay.Append(TagSource.ToString());
+			StringToDisplay.Append(TEXT("?"));
+
+			FNotificationInfo Info(FText::FromString(StringToDisplay));
+			Info.ExpireDuration = 10.f;
+			Info.ButtonDetails.Add(FNotificationButtonInfo(LOCTEXT("RestrictedTagPopupButtonAccept", "Yes"), FText(), FSimpleDelegate::CreateSP(this, &SAddNewRestrictedGameplayTagWidget::CreateNewRestrictedGameplayTag), SNotificationItem::CS_None));
+			Info.ButtonDetails.Add(FNotificationButtonInfo(LOCTEXT("RestrictedTagPopupButtonReject", "No"), FText(), FSimpleDelegate::CreateSP(this, &SAddNewRestrictedGameplayTagWidget::CancelNewTag), SNotificationItem::CS_None));
 
-		AddRestrictedGameplayTagDialog = FSlateNotificationManager::Get().AddNotification(Info);
+			AddRestrictedGameplayTagDialog = FSlateNotificationManager::Get().AddNotification(Info);
+		}
 	}
 	else
 	{
diff --git a/Engine/Source/Runtime/GameplayTags/Private/GameplayTagsSettings.cpp b/Engine/Source/Runtime/GameplayTags/Private/GameplayTagsSettings.cpp
--- a/Engine/Source/Runtime/GameplayTags/Private/GameplayTagsSettings.cpp
+++ b/Engine/Source/Runtime/GameplayTags/Private/GameplayTagsSettings.cpp
@@ -89,6 +89,10 @@ void UGameplayTagsSettings::PostEditChangeProperty(FPropertyChangedEvent& Proper
 			{
 				if (!Info.RestrictedConfigName.IsEmpty())
 				{
+					if (!Info.RestrictedConfigName.EndsWith(TEXT(".ini")))
+					{
+						Info.RestrictedConfigName.Append(TEXT(".ini"));
+					}
 					FGameplayTagSource* Source = Manager.FindOrAddTagSource(*Info.RestrictedConfigName, EGameplayTagSourceType::RestrictedTagList);
 					if (!Source)
 					{
