Remove delegates coming from gameplay tag widgets when the widgets are destroyed.

diff --git a/Engine/Plugins/Editor/GameplayTagsEditor/Source/GameplayTagsEditor/Private/SAddNewGameplayTagWidget.cpp b/Engine/Plugins/Editor/GameplayTagsEditor/Source/GameplayTagsEditor/Private/SAddNewGameplayTagWidget.cpp
--- a/Engine/Plugins/Editor/GameplayTagsEditor/Source/GameplayTagsEditor/Private/SAddNewGameplayTagWidget.cpp
+++ b/Engine/Plugins/Editor/GameplayTagsEditor/Source/GameplayTagsEditor/Private/SAddNewGameplayTagWidget.cpp
@@ -9,6 +9,14 @@
 
 #define LOCTEXT_NAMESPACE "AddNewGameplayTagWidget"
 
+SAddNewGameplayTagWidget::~SAddNewGameplayTagWidget()
+{
+	if (!GExitPurge)
+	{
+		IGameplayTagsModule::OnTagSettingsChanged.RemoveAll(this);
+	}
+}
+
 void SAddNewGameplayTagWidget::Construct(const FArguments& InArgs)
 {
 	
diff --git a/Engine/Plugins/Editor/GameplayTagsEditor/Source/GameplayTagsEditor/Private/SAddNewGameplayTagWidget.h b/Engine/Plugins/Editor/GameplayTagsEditor/Source/GameplayTagsEditor/Private/SAddNewGameplayTagWidget.h
--- a/Engine/Plugins/Editor/GameplayTagsEditor/Source/GameplayTagsEditor/Private/SAddNewGameplayTagWidget.h
+++ b/Engine/Plugins/Editor/GameplayTagsEditor/Source/GameplayTagsEditor/Private/SAddNewGameplayTagWidget.h
@@ -25,6 +25,8 @@ public:
 		SLATE_ARGUMENT( FString, NewTagName ) // String that will initially populate the New Tag Name field
 	SLATE_END_ARGS();
 
+	virtual ~SAddNewGameplayTagWidget();
+
 	virtual void Tick( const FGeometry& AllottedGeometry, const double InCurrentTime, const float InDeltaTime ) override;
 
 	void Construct( const FArguments& InArgs);
diff --git a/Engine/Plugins/Editor/GameplayTagsEditor/Source/GameplayTagsEditor/Private/SAddNewRestrictedGameplayTagWidget.cpp b/Engine/Plugins/Editor/GameplayTagsEditor/Source/GameplayTagsEditor/Private/SAddNewRestrictedGameplayTagWidget.cpp
--- a/Engine/Plugins/Editor/GameplayTagsEditor/Source/GameplayTagsEditor/Private/SAddNewRestrictedGameplayTagWidget.cpp
+++ b/Engine/Plugins/Editor/GameplayTagsEditor/Source/GameplayTagsEditor/Private/SAddNewRestrictedGameplayTagWidget.cpp
@@ -12,6 +12,14 @@
 
 #define LOCTEXT_NAMESPACE "AddNewRestrictedGameplayTagWidget"
 
+SAddNewRestrictedGameplayTagWidget::~SAddNewRestrictedGameplayTagWidget()
+{
+	if (!GExitPurge)
+	{
+		IGameplayTagsModule::OnTagSettingsChanged.RemoveAll(this);
+	}
+}
+
 void SAddNewRestrictedGameplayTagWidget::Construct(const FArguments& InArgs)
 {
 	FText HintText = LOCTEXT("NewTagNameHint", "X.Y.Z");
diff --git a/Engine/Plugins/Editor/GameplayTagsEditor/Source/GameplayTagsEditor/Private/SAddNewRestrictedGameplayTagWidget.h b/Engine/Plugins/Editor/GameplayTagsEditor/Source/GameplayTagsEditor/Private/SAddNewRestrictedGameplayTagWidget.h
--- a/Engine/Plugins/Editor/GameplayTagsEditor/Source/GameplayTagsEditor/Private/SAddNewRestrictedGameplayTagWidget.h
+++ b/Engine/Plugins/Editor/GameplayTagsEditor/Source/GameplayTagsEditor/Private/SAddNewRestrictedGameplayTagWidget.h
@@ -27,6 +27,8 @@ public:
 		SLATE_ARGUMENT( FString, NewRestrictedTagName ) // String that will initially populate the New Tag Name field
 	SLATE_END_ARGS();
 
+	virtual ~SAddNewRestrictedGameplayTagWidget();
+
 	virtual void Tick( const FGeometry& AllottedGeometry, const double InCurrentTime, const float InDeltaTime ) override;
 
 	void Construct( const FArguments& InArgs);
