Fix GameplayTags filtering in BP SET nodes.
marked GetCategoriesMetaFromStruct deprecated in favor of GetCategoriesMetaFromField
PR #5348: UE-67193: Retrieve filter string for BP GameplayTags set nodes (Contributed by [REDACTED_PERSON])

diff --git a/Engine/Plugins/Editor/GameplayTagsEditor/Source/GameplayTagsEditor/Private/SGameplayTagGraphPin.cpp b/Engine/Plugins/Editor/GameplayTagsEditor/Source/GameplayTagsEditor/Private/SGameplayTagGraphPin.cpp
--- a/Engine/Plugins/Editor/GameplayTagsEditor/Source/GameplayTagsEditor/Private/SGameplayTagGraphPin.cpp
+++ b/Engine/Plugins/Editor/GameplayTagsEditor/Source/GameplayTagsEditor/Private/SGameplayTagGraphPin.cpp
@@ -5,6 +5,7 @@
 #include "K2Node_CallFunction.h"
 #include "GameplayTagsModule.h"
 #include "Widgets/Layout/SScaleBox.h"
+#include "K2Node_VariableSet.h"
 
 #define LOCTEXT_NAMESPACE "GameplayTagGraphPin"
 
@@ -48,7 +49,7 @@ void SGameplayTagGraphPin::ParseDefaultValueData()
 	FilterString.Empty();
 	if (UScriptStruct* PinStructType = Cast<UScriptStruct>(GraphPinObj->PinType.PinSubCategoryObject.Get()))
 	{
-		FilterString = UGameplayTagsManager::Get().GetCategoriesMetaFromStruct(PinStructType);
+		FilterString = UGameplayTagsManager::Get().GetCategoriesMetaFromField(PinStructType);
 	}
 
 	if (FilterString.IsEmpty())
@@ -60,6 +61,10 @@ void SGameplayTagGraphPin::ParseDefaultValueData()
 				FilterString = UGameplayTagsManager::Get().GetCategoriesMetaFromFunction(ThisFunction);
 			}
 		}
+		else if (UK2Node_VariableSet* ThisVariable = Cast<UK2Node_VariableSet>(GraphPinObj->GetOwningNode()))
+		{
+			FilterString = UGameplayTagsManager::Get().GetCategoriesMetaFromField(ThisVariable->GetPropertyForVariable());
+		}
 	}
 
 	if (TagString.StartsWith(TEXT("(")) && TagString.EndsWith(TEXT(")")))
diff --git a/Engine/Source/Runtime/GameplayTags/Classes/GameplayTagsManager.h b/Engine/Source/Runtime/GameplayTags/Classes/GameplayTagsManager.h
--- a/Engine/Source/Runtime/GameplayTags/Classes/GameplayTagsManager.h
+++ b/Engine/Source/Runtime/GameplayTags/Classes/GameplayTagsManager.h
@@ -509,8 +509,12 @@ class GAMEPLAYTAGS_API UGameplayTagsManager : public UObject
 	/** Returns "Categories" meta property from given handle, used for filtering by tag widget */
 	FString GetCategoriesMetaFromPropertyHandle(TSharedPtr<class IPropertyHandle> PropertyHandle) const;
 
+	/** Returns "Categories" meta property from given field, used for filtering by tag widget */
+	FString GetCategoriesMetaFromField(UField* Field) const;
+
 	/** Returns "Categories" meta property from given struct, used for filtering by tag widget */
-	FString GetCategoriesMetaFromStruct(UScriptStruct* Struct) const;
+	UE_DEPRECATED(4.22, "Please call GetCategoriesMetaFromField instead.")
+	FString GetCategoriesMetaFromStruct(UScriptStruct* Struct) const { return GetCategoriesMetaFromField(Struct); }
 
 	/** Returns "GameplayTagFilter" meta property from given function, used for filtering by tag widget for any parameters of the function that end up as BP pins */
 	FString GetCategoriesMetaFromFunction(UFunction* Func) const;
diff --git a/Engine/Source/Runtime/GameplayTags/Private/GameplayTagsManager.cpp b/Engine/Source/Runtime/GameplayTags/Private/GameplayTagsManager.cpp
--- a/Engine/Source/Runtime/GameplayTags/Private/GameplayTagsManager.cpp
+++ b/Engine/Source/Runtime/GameplayTags/Private/GameplayTagsManager.cpp
@@ -1118,12 +1118,12 @@ void UGameplayTagsManager::GetFilteredGameplayRootTags(const FString& InFilterSt
 	}
 }
 
-FString UGameplayTagsManager::GetCategoriesMetaFromStruct(UScriptStruct* Struct) const
+FString UGameplayTagsManager::GetCategoriesMetaFromField(UField* Field) const
 {
-	check(Struct);
-	if (Struct->HasMetaData(NAME_Categories))
+	check(Field);
+	if (Field->HasMetaData(NAME_Categories))
 	{
-		return Struct->GetMetaData(NAME_Categories);
+		return Field->GetMetaData(NAME_Categories);
 	}
 	return FString();
 }
