Gameplay Tags - The editor customization now allows the easy removal of specific gameplay tags from the container.

diff --git a/Engine/Plugins/Editor/GameplayTagsEditor/Source/GameplayTagsEditor/Private/GameplayTagContainerCustomization.cpp b/Engine/Plugins/Editor/GameplayTagsEditor/Source/GameplayTagsEditor/Private/GameplayTagContainerCustomization.cpp
--- a/Engine/Plugins/Editor/GameplayTagsEditor/Source/GameplayTagsEditor/Private/GameplayTagContainerCustomization.cpp
+++ b/Engine/Plugins/Editor/GameplayTagsEditor/Source/GameplayTagsEditor/Private/GameplayTagContainerCustomization.cpp
@@ -11,6 +11,7 @@
 #include "DetailWidgetRow.h"
 #include "ScopedTransaction.h"
 #include "Widgets/Input/SHyperlink.h"
+#include "EditorFontGlyphs.h"
 
 #define LOCTEXT_NAMESPACE "GameplayTagContainerCustomization"
 
@@ -117,21 +118,47 @@ void FGameplayTagContainerCustomization::RefreshTagList()
 
 TSharedRef<ITableRow> FGameplayTagContainerCustomization::MakeListViewWidget(TSharedPtr<FString> Item, const TSharedRef<STableViewBase>& OwnerTable)
 {
+	TSharedPtr<SWidget> TagItem;
+
 	if (UGameplayTagsManager::Get().ShowGameplayTagAsHyperLinkEditor(*Item.Get()))
 	{
-		return SNew( STableRow< TSharedPtr<FString> >, OwnerTable )
-		[
-			SNew(SHyperlink)
-			.Text( FText::FromString(*Item.Get()) )
-			.OnNavigate( this, &FGameplayTagContainerCustomization::OnTagDoubleClicked, *Item.Get() )
-		];
-
+		TagItem = SNew(SHyperlink)
+			.Text(FText::FromString(*Item.Get()))
+			.OnNavigate(this, &FGameplayTagContainerCustomization::OnTagDoubleClicked, *Item.Get());
+	}
+	else
+	{
+		TagItem = SNew(STextBlock).Text(FText::FromString(*Item.Get()));
 	}
-
 
 	return SNew( STableRow< TSharedPtr<FString> >, OwnerTable )
 	[
-		SNew(STextBlock) .Text( FText::FromString(*Item.Get()) )
+		SNew(SHorizontalBox)
+
+		+ SHorizontalBox::Slot()
+		.AutoWidth()
+		.VAlign(VAlign_Center)
+		.Padding(0,0,2,0)
+		[
+			SNew(SButton)
+			.IsEnabled(!StructPropertyHandle->IsEditConst())
+			.ContentPadding(FMargin(0))
+			.ButtonStyle(FEditorStyle::Get(), "FlatButton.Danger")
+			.ForegroundColor(FSlateColor::UseForeground())
+			.OnClicked(this, &FGameplayTagContainerCustomization::OnRemoveTagClicked, *Item.Get())
+			[
+				SNew(STextBlock)
+				.Font(FEditorStyle::Get().GetFontStyle("FontAwesome.9"))
+				.Text(FEditorFontGlyphs::Times)
+			]
+		]
+
+		+ SHorizontalBox::Slot()
+		.AutoWidth()
+		.VAlign(VAlign_Center)
+		[
+			TagItem.ToSharedRef()
+		]
 	];
 }
 
@@ -140,6 +167,40 @@ void FGameplayTagContainerCustomization::OnTagDoubleClicked(FString TagName)
 	UGameplayTagsManager::Get().NotifyGameplayTagDoubleClickedEditor(TagName);
 }
 
+FReply FGameplayTagContainerCustomization::OnRemoveTagClicked(FString TagName)
+{
+	TArray<FString> NewValues;
+
+	StructPropertyHandle->EnumerateRawData([TagName, &NewValues](void* RawTagContainer, const int32 /*DataIndex*/, const int32 /*NumDatas*/) -> bool {
+		
+		FGameplayTagContainer TagContainerCopy = *static_cast<FGameplayTagContainer*>(RawTagContainer);
+		for (int32 TagIndex = 0; TagIndex < TagContainerCopy.Num(); TagIndex++)
+		{
+			FGameplayTag Tag = TagContainerCopy.GetByIndex(TagIndex);
+			if (Tag.GetTagName().ToString() == TagName)
+			{
+				TagContainerCopy.RemoveTag(Tag);
+			}
+		}
+		
+		NewValues.Add(TagContainerCopy.ToString());
+
+		return true;
+	});
+
+	{
+		FScopedTransaction Transaction(LOCTEXT("RemoveGameplayTagFromContainer", "Remove Gameplay Tag"));
+		for (int i = 0; i < NewValues.Num(); i++)
+		{
+			StructPropertyHandle->SetPerObjectValue(i, NewValues[i]);
+		}
+	}
+
+	RefreshTagList();
+
+	return FReply::Handled();
+}
+
 TSharedRef<SWidget> FGameplayTagContainerCustomization::GetListContent()
 {
 	if (!StructPropertyHandle.IsValid() || StructPropertyHandle->GetProperty() == nullptr)
diff --git a/Engine/Plugins/Editor/GameplayTagsEditor/Source/GameplayTagsEditor/Private/GameplayTagContainerCustomization.h b/Engine/Plugins/Editor/GameplayTagsEditor/Source/GameplayTagsEditor/Private/GameplayTagContainerCustomization.h
--- a/Engine/Plugins/Editor/GameplayTagsEditor/Source/GameplayTagsEditor/Private/GameplayTagContainerCustomization.h
+++ b/Engine/Plugins/Editor/GameplayTagsEditor/Source/GameplayTagsEditor/Private/GameplayTagContainerCustomization.h
@@ -74,5 +74,6 @@ private:
 	TSharedPtr<SListView<TSharedPtr<FString>>> TagListView;
 
 	void OnTagDoubleClicked(FString TagName);
+	FReply OnRemoveTagClicked(FString TagName);
 };
 
