Gameplay Tag - When the gameplay tag dialog opens we now focus the search field.

diff --git a/Engine/Plugins/Editor/GameplayTagsEditor/Source/GameplayTagsEditor/Private/GameplayTagContainerCustomization.cpp b/Engine/Plugins/Editor/GameplayTagsEditor/Source/GameplayTagsEditor/Private/GameplayTagContainerCustomization.cpp
--- a/Engine/Plugins/Editor/GameplayTagsEditor/Source/GameplayTagsEditor/Private/GameplayTagContainerCustomization.cpp
+++ b/Engine/Plugins/Editor/GameplayTagsEditor/Source/GameplayTagsEditor/Private/GameplayTagContainerCustomization.cpp
@@ -41,8 +41,9 @@ void FGameplayTagContainerCustomization::CustomizeHeader(TSharedRef<class IPrope
 				+ SVerticalBox::Slot()
 				.AutoHeight()
 				[
-					SNew(SComboButton)
+					SAssignNew(EditButton, SComboButton)
 					.OnGetMenuContent(this, &FGameplayTagContainerCustomization::GetListContent)
+					.OnMenuOpenChanged(this, &FGameplayTagContainerCustomization::OnGameplayTagListMenuOpenStateChanged)
 					.ContentPadding(FMargin(2.0f, 2.0f))
 					.MenuPlacement(MenuPlacement_BelowAnchor)
 					.ButtonContent()
@@ -215,20 +216,36 @@ TSharedRef<SWidget> FGameplayTagContainerCustomization::GetListContent()
 
 	bool bReadOnly = StructPropertyHandle->IsEditConst();
 
+	TSharedRef<SGameplayTagWidget> TagWidget = SNew(SGameplayTagWidget, EditableContainers)
+		.Filter(Categories)
+		.ReadOnly(bReadOnly)
+		.TagContainerName(StructPropertyHandle->GetPropertyDisplayName().ToString())
+		.OnTagChanged(this, &FGameplayTagContainerCustomization::RefreshTagList)
+		.PropertyHandle(StructPropertyHandle);
+
+	LastTagWidget = TagWidget;
+
 	return SNew(SVerticalBox)
 		+ SVerticalBox::Slot()
 		.AutoHeight()
 		.MaxHeight(400)
 		[
-			SNew(SGameplayTagWidget, EditableContainers)
-			.Filter(Categories)
-			.ReadOnly(bReadOnly)
-			.TagContainerName(StructPropertyHandle->GetPropertyDisplayName().ToString())
-			.OnTagChanged(this, &FGameplayTagContainerCustomization::RefreshTagList)
-			.PropertyHandle(StructPropertyHandle)
+			TagWidget
 		];
 }
 
+void FGameplayTagContainerCustomization::OnGameplayTagListMenuOpenStateChanged(bool bIsOpened)
+{
+	if (bIsOpened)
+	{
+		TSharedPtr<SGameplayTagWidget> TagWidget = LastTagWidget.Pin();
+		if (TagWidget.IsValid())
+		{
+			EditButton->SetMenuContentWidgetToFocus(TagWidget->GetWidgetToFocusOnOpen());
+		}
+	}
+}
+
 FReply FGameplayTagContainerCustomization::OnClearAllButtonClicked()
 {
 	FScopedTransaction Transaction(LOCTEXT("GameplayTagContainerCustomization_RemoveAllTags", "Remove All Gameplay Tags"));
diff --git a/Engine/Plugins/Editor/GameplayTagsEditor/Source/GameplayTagsEditor/Private/GameplayTagContainerCustomization.h b/Engine/Plugins/Editor/GameplayTagsEditor/Source/GameplayTagsEditor/Private/GameplayTagContainerCustomization.h
--- a/Engine/Plugins/Editor/GameplayTagsEditor/Source/GameplayTagsEditor/Private/GameplayTagContainerCustomization.h
+++ b/Engine/Plugins/Editor/GameplayTagsEditor/Source/GameplayTagsEditor/Private/GameplayTagContainerCustomization.h
@@ -61,6 +61,8 @@ private:
 	/** Callback function to create content for the combo button. */
 	TSharedRef<SWidget> GetListContent();
 
+	void OnGameplayTagListMenuOpenStateChanged(bool bIsOpened);
+
 	/** Cached property handle */
 	TSharedPtr<IPropertyHandle> StructPropertyHandle;
 
@@ -73,6 +75,10 @@ private:
 	/** The TagList, kept as a member so we can update it later */
 	TSharedPtr<SListView<TSharedPtr<FString>>> TagListView;
 
+	TSharedPtr<class SComboButton> EditButton;
+
+	TWeakPtr<class SGameplayTagWidget> LastTagWidget;
+
 	void OnTagDoubleClicked(FString TagName);
 	FReply OnRemoveTagClicked(FString TagName);
 };
diff --git a/Engine/Plugins/Editor/GameplayTagsEditor/Source/GameplayTagsEditor/Private/GameplayTagCustomization.cpp b/Engine/Plugins/Editor/GameplayTagsEditor/Source/GameplayTagsEditor/Private/GameplayTagCustomization.cpp
--- a/Engine/Plugins/Editor/GameplayTagsEditor/Source/GameplayTagsEditor/Private/GameplayTagCustomization.cpp
+++ b/Engine/Plugins/Editor/GameplayTagsEditor/Source/GameplayTagsEditor/Private/GameplayTagCustomization.cpp
@@ -38,8 +38,9 @@ void FGameplayTagCustomization::CustomizeHeader(TSharedRef<class IPropertyHandle
 		+ SHorizontalBox::Slot()
 		.AutoWidth()
 		[
-			SNew(SComboButton)
+			SAssignNew(EditButton, SComboButton)
 			.OnGetMenuContent(this, &FGameplayTagCustomization::GetListContent)
+			.OnMenuOpenChanged(this, &FGameplayTagCustomization::OnGameplayTagListMenuOpenStateChanged)
 			.ContentPadding(FMargin(2.0f, 2.0f))
 			.MenuPlacement(MenuPlacement_BelowAnchor)
 			.ButtonContent()
@@ -96,21 +97,38 @@ TSharedRef<SWidget> FGameplayTagCustomization::GetListContent()
 
 	bool bReadOnly = StructPropertyHandle->IsEditConst();
 
+	TSharedRef<SGameplayTagWidget> TagWidget =
+		SNew(SGameplayTagWidget, EditableContainers)
+		.Filter(Categories)
+		.ReadOnly(bReadOnly)
+		.TagContainerName(StructPropertyHandle->GetPropertyDisplayName().ToString())
+		.MultiSelect(false)
+		.OnTagChanged(this, &FGameplayTagCustomization::OnTagChanged)
+		.PropertyHandle(StructPropertyHandle);
+
+	LastTagWidget = TagWidget;
+
 	return SNew(SVerticalBox)
 		+ SVerticalBox::Slot()
 		.AutoHeight()
 		.MaxHeight(400)
 		[
-			SNew(SGameplayTagWidget, EditableContainers)
-			.Filter(Categories)
-			.ReadOnly(bReadOnly)
-			.TagContainerName(StructPropertyHandle->GetPropertyDisplayName().ToString())
-			.MultiSelect(false)
-			.OnTagChanged(this, &FGameplayTagCustomization::OnTagChanged)
-			.PropertyHandle(StructPropertyHandle)
+			TagWidget
 		];
 }
 
+void FGameplayTagCustomization::OnGameplayTagListMenuOpenStateChanged(bool bIsOpened)
+{
+	if (bIsOpened)
+	{
+		TSharedPtr<SGameplayTagWidget> TagWidget = LastTagWidget.Pin();
+		if (TagWidget.IsValid())
+		{
+			EditButton->SetMenuContentWidgetToFocus(TagWidget->GetWidgetToFocusOnOpen());
+		}
+	}
+}
+
 void FGameplayTagCustomization::OnPropertyValueChanged()
 {
 	TagName = TEXT("");
diff --git a/Engine/Plugins/Editor/GameplayTagsEditor/Source/GameplayTagsEditor/Private/GameplayTagCustomization.h b/Engine/Plugins/Editor/GameplayTagsEditor/Source/GameplayTagsEditor/Private/GameplayTagCustomization.h
--- a/Engine/Plugins/Editor/GameplayTagsEditor/Source/GameplayTagsEditor/Private/GameplayTagCustomization.h
+++ b/Engine/Plugins/Editor/GameplayTagsEditor/Source/GameplayTagsEditor/Private/GameplayTagCustomization.h
@@ -42,6 +42,8 @@ private:
 	/** Callback function to create content for the combo button. */
 	TSharedRef<SWidget> GetListContent();
 
+	void OnGameplayTagListMenuOpenStateChanged(bool bIsOpened);
+
 	/** Returns Tag name currently selected*/
 	FText SelectedTag() const;
 
@@ -59,5 +61,9 @@ private:
 
 	void OnTagDoubleClicked();
 	EVisibility GetVisibilityForTagTextBlockWidget(bool ForTextWidget) const;
+
+	TSharedPtr<class SComboButton> EditButton;
+
+	TWeakPtr<class SGameplayTagWidget> LastTagWidget;
 };
 
diff --git a/Engine/Plugins/Editor/GameplayTagsEditor/Source/GameplayTagsEditor/Private/SGameplayTagWidget.cpp b/Engine/Plugins/Editor/GameplayTagsEditor/Source/GameplayTagsEditor/Private/SGameplayTagWidget.cpp
--- a/Engine/Plugins/Editor/GameplayTagsEditor/Source/GameplayTagsEditor/Private/SGameplayTagWidget.cpp
+++ b/Engine/Plugins/Editor/GameplayTagsEditor/Source/GameplayTagsEditor/Private/SGameplayTagWidget.cpp
@@ -1368,4 +1368,9 @@ void SGameplayTagWidget::OpenRenameGameplayTagDialog(TSharedPtr<FGameplayTagNode
 	FSlateApplication::Get().AddModalWindow(RenameTagWindow, CurrentWindow);
 }
 
+TSharedPtr<SWidget> SGameplayTagWidget::GetWidgetToFocusOnOpen()
+{
+	return SearchTagBox;
+}
+
 #undef LOCTEXT_NAMESPACE
diff --git a/Engine/Plugins/Editor/GameplayTagsEditor/Source/GameplayTagsEditor/Private/SGameplayTagWidget.h b/Engine/Plugins/Editor/GameplayTagsEditor/Source/GameplayTagsEditor/Private/SGameplayTagWidget.h
--- a/Engine/Plugins/Editor/GameplayTagsEditor/Source/GameplayTagsEditor/Private/SGameplayTagWidget.h
+++ b/Engine/Plugins/Editor/GameplayTagsEditor/Source/GameplayTagsEditor/Private/SGameplayTagWidget.h
@@ -98,6 +98,9 @@ public:
 	/** Forces the widget to refresh its tags on the next tick */
 	void RefreshOnNextTick();
 
+	/** Gets the widget to focus once the menu opens. */
+	TSharedPtr<SWidget> GetWidgetToFocusOnOpen();
+
 private:
 
 	/* string that sets the section of the ini file to use for this class*/ 
