Make gameplay tag queries call NotifyPreChange and NotifyPostChange before and after changes are saved in the editor.
Add calls to NotifyPreChange and NotifyPostChange to the "Clear All" button for gameplay tag queries.

diff --git a/Engine/Plugins/Editor/GameplayTagsEditor/Source/GameplayTagsEditor/Private/GameplayTagQueryCustomization.cpp b/Engine/Plugins/Editor/GameplayTagsEditor/Source/GameplayTagsEditor/Private/GameplayTagQueryCustomization.cpp
--- a/Engine/Plugins/Editor/GameplayTagsEditor/Source/GameplayTagsEditor/Private/GameplayTagQueryCustomization.cpp
+++ b/Engine/Plugins/Editor/GameplayTagsEditor/Source/GameplayTagsEditor/Private/GameplayTagQueryCustomization.cpp
@@ -31,49 +31,49 @@ void FGameplayTagQueryCustomization::CustomizeHeader(TSharedRef<class IPropertyH
 	bool const bReadOnly = StructPropertyHandle->IsEditConst();
 
 	HeaderRow
-	.NameContent()
-	[
-		StructPropertyHandle->CreatePropertyNameWidget()
-	]
+		.NameContent()
+		[
+			StructPropertyHandle->CreatePropertyNameWidget()
+		]
 	.ValueContent()
-	.MaxDesiredWidth(512)
-	[
-		SNew(SHorizontalBox)
+		.MaxDesiredWidth(512)
+		[
+			SNew(SHorizontalBox)
 		+SHorizontalBox::Slot()
 		.AutoWidth()
 		.VAlign(VAlign_Center)
 		[
 			SNew(SVerticalBox)
 			+SVerticalBox::Slot()
-			.AutoHeight()
-			[
-				SNew(SButton)
-				.Text(this, &FGameplayTagQueryCustomization::GetEditButtonText)
-				.OnClicked(this, &FGameplayTagQueryCustomization::OnEditButtonClicked)
-			]
+		.AutoHeight()
+		[
+			SNew(SButton)
+			.Text(this, &FGameplayTagQueryCustomization::GetEditButtonText)
+		.OnClicked(this, &FGameplayTagQueryCustomization::OnEditButtonClicked)
+		]
 			+SVerticalBox::Slot()
-			.AutoHeight()
-			[
-				SNew(SButton)
-				.IsEnabled(!bReadOnly)
-				.Text(LOCTEXT("GameplayTagQueryCustomization_Clear", "Clear All"))
-				.OnClicked(this, &FGameplayTagQueryCustomization::OnClearAllButtonClicked)
-				.Visibility(this, &FGameplayTagQueryCustomization::GetClearAllVisibility)
-			]
+		.AutoHeight()
+		[
+			SNew(SButton)
+			.IsEnabled(!bReadOnly)
+		.Text(LOCTEXT("GameplayTagQueryCustomization_Clear", "Clear All"))
+		.OnClicked(this, &FGameplayTagQueryCustomization::OnClearAllButtonClicked)
+		.Visibility(this, &FGameplayTagQueryCustomization::GetClearAllVisibility)
+		]
 		]
 		+SHorizontalBox::Slot()
 		.AutoWidth()
 		[
 			SNew(SBorder)
 			.Padding(4.0f)
-			.Visibility(this, &FGameplayTagQueryCustomization::GetQueryDescVisibility)
-			[
-				SNew(STextBlock)
-				.Text(this, &FGameplayTagQueryCustomization::GetQueryDescText)
-				.AutoWrapText(true)
-			]
+		.Visibility(this, &FGameplayTagQueryCustomization::GetQueryDescVisibility)
+		[
+			SNew(STextBlock)
+			.Text(this, &FGameplayTagQueryCustomization::GetQueryDescText)
+		.AutoWrapText(true)
+		]
 		]
-	];
+		];
 
 	GEditor->RegisterForUndo(this);
 }
@@ -99,6 +99,11 @@ FText FGameplayTagQueryCustomization::GetEditButtonText() const
 
 FReply FGameplayTagQueryCustomization::OnClearAllButtonClicked()
 {
+	if (StructPropertyHandle.IsValid())
+	{
+		StructPropertyHandle->NotifyPreChange();
+	}
+
 	for (auto& EQ : EditableQueries)
 	{
 		if (EQ.TagQuery)
@@ -109,6 +114,11 @@ FReply FGameplayTagQueryCustomization::OnClearAllButtonClicked()
 
 	RefreshQueryDescription();
 
+	if (StructPropertyHandle.IsValid())
+	{
+		StructPropertyHandle->NotifyPostChange();
+	}
+
 	return FReply::Handled();
 }
 
@@ -145,7 +155,7 @@ void FGameplayTagQueryCustomization::RefreshQueryDescription()
 	{
 		QueryDescription = TEXT("Multiple Selected");
 	}
-	else if ( (EditableQueries.Num() == 1) && (EditableQueries[0].TagQuery != nullptr) )
+	else if ((EditableQueries.Num() == 1) && (EditableQueries[0].TagQuery != nullptr))
 	{
 		QueryDescription = EditableQueries[0].TagQuery->GetDescription();
 	}
@@ -188,9 +198,10 @@ FReply FGameplayTagQueryCustomization::OnEditButtonClicked()
 				.ClientSize(FVector2D(600, 400))
 				[
 					SNew(SGameplayTagQueryWidget, EditableQueries)
-					.OnSaveAndClose(this, &FGameplayTagQueryCustomization::CloseWidgetWindow, false)
-					.OnCancel(this, &FGameplayTagQueryCustomization::CloseWidgetWindow, true)
-					.ReadOnly(bReadOnly)
+					.OnClosePreSave(this, &FGameplayTagQueryCustomization::PreSave)
+				.OnSaveAndClose(this, &FGameplayTagQueryCustomization::CloseWidgetWindow, false)
+				.OnCancel(this, &FGameplayTagQueryCustomization::CloseWidgetWindow, true)
+				.ReadOnly(bReadOnly)
 				];
 
 			// NOTE: FGlobalTabmanager::Get()-> is actually dereferencing a SharedReference, not a SharedPtr, so it cannot be null.
@@ -239,35 +250,29 @@ void FGameplayTagQueryCustomization::BuildEditableQueryList()
 	}	
 }
 
-void FGameplayTagQueryCustomization::CloseWidgetWindow(bool WasCancelled)
+void FGameplayTagQueryCustomization::PreSave()
 {
-	// Notify change. This is required for these to work inside of UDataTables
-	if (!WasCancelled && PropertyUtilities.IsValid())
+	if (StructPropertyHandle.IsValid())
 	{
-		if (StructPropertyHandle.IsValid())
-		{
-			UProperty* TheProperty = StructPropertyHandle->GetProperty();
-
-			FEditPropertyChain PropertyChain;
-			PropertyChain.AddHead(TheProperty);
-			PropertyChain.SetActivePropertyNode(TheProperty);
+		StructPropertyHandle->NotifyPreChange();
+	}
+}
 
-			FPropertyChangedEvent ChangeEvent(StructPropertyHandle->GetProperty(), EPropertyChangeType::ValueSet, nullptr);
-			FNotifyHook* NotifyHook = PropertyUtilities->GetNotifyHook();
-			if (NotifyHook != nullptr)
-			{
-				NotifyHook->NotifyPostChange(ChangeEvent, &PropertyChain);
-			}
-		}
+void FGameplayTagQueryCustomization::CloseWidgetWindow(bool WasCancelled)
+{
+	// Notify change.
+	if (!WasCancelled && StructPropertyHandle.IsValid())
+	{
+		StructPropertyHandle->NotifyPostChange();
 	}
 
- 	if (GameplayTagQueryWidgetWindow.IsValid())
- 	{
- 		GameplayTagQueryWidgetWindow->RequestDestroyWindow();
+	if (GameplayTagQueryWidgetWindow.IsValid())
+	{
+		GameplayTagQueryWidgetWindow->RequestDestroyWindow();
 		GameplayTagQueryWidgetWindow = nullptr;
 
 		RefreshQueryDescription();
- 	}
+	}
 }
 
-#undef LOCTEXT_NAMESPACE
+#undef LOCTEXT_NAMESPACE
\ No newline at end of file
diff --git a/Engine/Plugins/Editor/GameplayTagsEditor/Source/GameplayTagsEditor/Private/GameplayTagQueryCustomization.h b/Engine/Plugins/Editor/GameplayTagsEditor/Source/GameplayTagsEditor/Private/GameplayTagQueryCustomization.h
--- a/Engine/Plugins/Editor/GameplayTagsEditor/Source/GameplayTagsEditor/Private/GameplayTagQueryCustomization.h
+++ b/Engine/Plugins/Editor/GameplayTagsEditor/Source/GameplayTagsEditor/Private/GameplayTagQueryCustomization.h
@@ -48,6 +48,8 @@ private:
 
 	FText GetQueryDescText() const;
 
+	void PreSave();
+
 	void CloseWidgetWindow(bool WasCancelled);
 
 	/** Build List of Editable Queries */
diff --git a/Engine/Plugins/Editor/GameplayTagsEditor/Source/GameplayTagsEditor/Private/SGameplayTagQueryWidget.cpp b/Engine/Plugins/Editor/GameplayTagsEditor/Source/GameplayTagsEditor/Private/SGameplayTagQueryWidget.cpp
--- a/Engine/Plugins/Editor/GameplayTagsEditor/Source/GameplayTagsEditor/Private/SGameplayTagQueryWidget.cpp
+++ b/Engine/Plugins/Editor/GameplayTagsEditor/Source/GameplayTagsEditor/Private/SGameplayTagQueryWidget.cpp
@@ -18,6 +18,7 @@ void SGameplayTagQueryWidget::Construct(const FArguments& InArgs, const TArray<F
 
 	bReadOnly = InArgs._ReadOnly;
 	bAutoSave = InArgs._AutoSave;
+	OnClosePreSave = InArgs._OnClosePreSave;
 	OnSaveAndClose = InArgs._OnSaveAndClose;
 	OnCancel = InArgs._OnCancel;
 	OnQueryChanged = InArgs._OnQueryChanged;
@@ -48,45 +49,46 @@ void SGameplayTagQueryWidget::Construct(const FArguments& InArgs, const TArray<F
 	Details->OnFinishedChangingProperties().AddSP(this, &SGameplayTagQueryWidget::OnFinishedChangingProperties);
 
 	ChildSlot
-	[
-		SNew(SBorder)
-		.BorderImage(FEditorStyle::GetBrush("ToolPanel.GroupBorder"))
+		[
+			SNew(SBorder)
+			.BorderImage(FEditorStyle::GetBrush("ToolPanel.GroupBorder"))
 		[
 			SNew(SVerticalBox)
 			+ SVerticalBox::Slot()
-			.AutoHeight()
-			.VAlign(VAlign_Top)
-			[
-				SNew(SHorizontalBox)
-				+ SHorizontalBox::Slot()
-				.AutoWidth()
-				[
-					SNew(SButton)
-					.IsEnabled(!bReadOnly)
-					.Visibility(this, &SGameplayTagQueryWidget::GetSaveAndCloseButtonVisibility)
-					.OnClicked(this, &SGameplayTagQueryWidget::OnSaveAndCloseClicked)
-					.Text(LOCTEXT("GameplayTagQueryWidget_SaveAndClose", "Save and Close"))
-				]
-				+ SHorizontalBox::Slot()
-				.AutoWidth()
-				[
-					SNew(SButton)
-					.Visibility(this, &SGameplayTagQueryWidget::GetCancelButtonVisibility)
-					.OnClicked(this, &SGameplayTagQueryWidget::OnCancelClicked)
-					.Text(LOCTEXT("GameplayTagQueryWidget_Cancel", "Close Without Saving"))
-				]
-			]
-			// to delete!
-			+ SVerticalBox::Slot()
-			[
-				Details.ToSharedRef()
-			]
+		.AutoHeight()
+		.VAlign(VAlign_Top)
+		[
+			SNew(SHorizontalBox)
+			+ SHorizontalBox::Slot()
+		.AutoWidth()
+		[
+			SNew(SButton)
+			.IsEnabled(!bReadOnly)
+		.Visibility(this, &SGameplayTagQueryWidget::GetSaveAndCloseButtonVisibility)
+		.OnClicked(this, &SGameplayTagQueryWidget::OnSaveAndCloseClicked)
+		.Text(LOCTEXT("GameplayTagQueryWidget_SaveAndClose", "Save and Close"))
+		]
+	+ SHorizontalBox::Slot()
+		.AutoWidth()
+		[
+			SNew(SButton)
+			.Visibility(this, &SGameplayTagQueryWidget::GetCancelButtonVisibility)
+		.OnClicked(this, &SGameplayTagQueryWidget::OnCancelClicked)
+		.Text(LOCTEXT("GameplayTagQueryWidget_Cancel", "Close Without Saving"))
 		]
-	];
+		]
+	// to delete!
+	+ SVerticalBox::Slot()
+		[
+			Details.ToSharedRef()
+		]
+		]
+		];
 }
 
 void SGameplayTagQueryWidget::OnFinishedChangingProperties(const FPropertyChangedEvent& PropertyChangedEvent)
 {
+	// Auto saved changes will not call pre and post notify; auto save should only be used to make changes coming from blueprints
 	if (bAutoSave)
 	{
 		SaveToTagQuery();
@@ -151,6 +153,8 @@ void SGameplayTagQueryWidget::SaveToTagQuery()
 
 FReply SGameplayTagQueryWidget::OnSaveAndCloseClicked()
 {
+	OnClosePreSave.ExecuteIfBound();
+
 	SaveToTagQuery();
 
 	OnSaveAndClose.ExecuteIfBound();
@@ -163,4 +167,4 @@ FReply SGameplayTagQueryWidget::OnCancelClicked()
 	return FReply::Handled();
 }
 
-#undef LOCTEXT_NAMESPACE
+#undef LOCTEXT_NAMESPACE
\ No newline at end of file
diff --git a/Engine/Plugins/Editor/GameplayTagsEditor/Source/GameplayTagsEditor/Private/SGameplayTagQueryWidget.h b/Engine/Plugins/Editor/GameplayTagsEditor/Source/GameplayTagsEditor/Private/SGameplayTagQueryWidget.h
--- a/Engine/Plugins/Editor/GameplayTagsEditor/Source/GameplayTagsEditor/Private/SGameplayTagQueryWidget.h
+++ b/Engine/Plugins/Editor/GameplayTagsEditor/Source/GameplayTagsEditor/Private/SGameplayTagQueryWidget.h
@@ -21,15 +21,16 @@ public:
 	SLATE_BEGIN_ARGS( SGameplayTagQueryWidget )
 		: _ReadOnly(false), _AutoSave(false)
 	{}
-		SLATE_ARGUMENT( bool, ReadOnly )	// Flag to set if the list is read only
-		SLATE_ARGUMENT( bool, AutoSave)		// Flag to set if edits should be applied automatically (hides buttons)
+		SLATE_ARGUMENT(bool, ReadOnly)		// Flag to set if the list is read only
+		SLATE_ARGUMENT(bool, AutoSave)		// Flag to set if edits should be applied automatically (hides buttons)
+		SLATE_EVENT(FSimpleDelegate, OnClosePreSave) // Called when "Save and Close" button clicked
 		SLATE_EVENT(FSimpleDelegate, OnSaveAndClose) // Called when "Save and Close" button clicked
 		SLATE_EVENT(FSimpleDelegate, OnCancel) // Called when "Close Without Saving" button clicked
 		SLATE_EVENT(FSimpleDelegate, OnQueryChanged)	// Called when the user has modified the query
-	SLATE_END_ARGS()
+		SLATE_END_ARGS()
 
-	/** Simple struct holding a tag query and its owner for generic re-use of the widget */
-	struct FEditableGameplayTagQueryDatum
+		/** Simple struct holding a tag query and its owner for generic re-use of the widget */
+		struct FEditableGameplayTagQueryDatum
 	{
 		/** Constructor */
 		FEditableGameplayTagQueryDatum(class UObject* InOwnerObj, struct FGameplayTagQuery* InTagQuery, FString* InTagExportText=nullptr)
@@ -64,13 +65,16 @@ private:
 	/** Containers to modify */
 	TArray<FEditableGameplayTagQueryDatum> TagQueries;
 
-	/** Called when "save and close" is clicked */
+	/** Called when "Save and Close" is clicked before we save the data. */
+	FSimpleDelegate OnClosePreSave;
+
+	/** Called when "Save and Close" is clicked after we have saved the data. */
 	FSimpleDelegate OnSaveAndClose;
 
 	/** Called when the user has modified the query */
 	FSimpleDelegate OnQueryChanged;
 
-	/** Called when "cancel" is clicked */
+	/** Called when "Close Without Saving" is clicked */
 	FSimpleDelegate OnCancel;
 
 	/** Properties Tab */
@@ -79,17 +83,17 @@ private:
 	class UEditableGameplayTagQuery* CreateEditableQuery(FGameplayTagQuery& Q);
 	TWeakObjectPtr<class UEditableGameplayTagQuery> EditableQuery;
 
-	/** Called when the user clicks the "Expand All" button; Expands the entire tag tree */
+	/** Called when the user clicks the "Save and Close" button */
 	FReply OnSaveAndCloseClicked();
 
-	/** Called when the user clicks the "Expand All" button; Expands the entire tag tree */
+	/** Called when the user clicks the "Close Without Saving" button */
 	FReply OnCancelClicked();
 
 	void OnFinishedChangingProperties(const FPropertyChangedEvent& PropertyChangedEvent);
 
-	/** Controls visibility of the "save and close" button */
+	/** Controls visibility of the "Save and Close" button */
 	EVisibility GetSaveAndCloseButtonVisibility() const;
-	/** Controls visibility of the "cancel without saving" button */
+	/** Controls visibility of the "Close Without Saving" button */
 	EVisibility GetCancelButtonVisibility() const;
 
 	void SaveToTagQuery();
