- Add gameplay tag filter support to tag container pins
- Support "UPARAM(meta=(Categories=" mark-up on UFUNCTION params for tag filtering (filtering will prefer UPARAM first, but if not found, will respect UFUNCTION's tag filter)
- Extract graph pin -> filter string function into new tag uitlities file

diff --git a/Engine/Plugins/Editor/GameplayTagsEditor/Source/GameplayTagsEditor/Private/GameplayTagPinUtilities.cpp b/Engine/Plugins/Editor/GameplayTagsEditor/Source/GameplayTagsEditor/Private/GameplayTagPinUtilities.cpp
new file mode 100644
--- /dev/null
+++ b/Engine/Plugins/Editor/GameplayTagsEditor/Source/GameplayTagsEditor/Private/GameplayTagPinUtilities.cpp
@@ -0,0 +1,38 @@
+// Copyright 1998-2019 Epic Games, Inc. All Rights Reserved.
+
+#include "GameplayTagPinUtilities.h"
+#include "GameplayTagsManager.h"
+#include "SGraphPin.h"
+#include "K2Node_CallFunction.h"
+#include "K2Node_VariableSet.h"
+
+FString GameplayTagPinUtilities::ExtractTagFilterStringFromGraphPin(UEdGraphPin* InTagPin)
+{
+	FString FilterString;
+
+	if (ensure(InTagPin))
+	{
+		const UGameplayTagsManager& TagManager = UGameplayTagsManager::Get();
+		if (UScriptStruct* PinStructType = Cast<UScriptStruct>(InTagPin->PinType.PinSubCategoryObject.Get()))
+		{
+			FilterString = TagManager.GetCategoriesMetaFromField(PinStructType);
+		}
+
+		if (FilterString.IsEmpty())
+		{
+			if (UK2Node_CallFunction* CallFuncNode = Cast<UK2Node_CallFunction>(InTagPin->GetOwningNode()))
+			{
+				if (UFunction* ThisFunction = CallFuncNode->GetTargetFunction())
+				{
+					FilterString = TagManager.GetCategoriesMetaFromFunction(ThisFunction, InTagPin->PinName);
+				}
+			}
+			else if (UK2Node_VariableSet* ThisVariable = Cast<UK2Node_VariableSet>(InTagPin->GetOwningNode()))
+			{
+				FilterString = TagManager.GetCategoriesMetaFromField(ThisVariable->GetPropertyForVariable());
+			}
+		}
+	}
+
+	return FilterString;
+}
diff --git a/Engine/Plugins/Editor/GameplayTagsEditor/Source/GameplayTagsEditor/Private/GameplayTagPinUtilities.h b/Engine/Plugins/Editor/GameplayTagsEditor/Source/GameplayTagsEditor/Private/GameplayTagPinUtilities.h
new file mode 100644
--- /dev/null
+++ b/Engine/Plugins/Editor/GameplayTagsEditor/Source/GameplayTagsEditor/Private/GameplayTagPinUtilities.h
@@ -0,0 +1,22 @@
+// Copyright 1998-2019 Epic Games, Inc. All Rights Reserved.
+
+#pragma once
+
+#include "CoreMinimal.h"
+
+// Forward declarations
+class UEdGraphPin;
+
+/** Utility functions for gameplay tag-related pins */
+namespace GameplayTagPinUtilities
+{
+	/**
+	 * Given a editor graph pin representing a tag or container, extract the appropriate filter
+	 * string from any associated metadata
+	 * 
+	 * @param InTagPin	Pin to extract from
+	 * 
+	 * @return Filter string, if any, to apply from the tag pin based on metadata
+	 */
+	FString ExtractTagFilterStringFromGraphPin(UEdGraphPin* InTagPin);
+}
diff --git a/Engine/Plugins/Editor/GameplayTagsEditor/Source/GameplayTagsEditor/Private/SGameplayTagContainerGraphPin.cpp b/Engine/Plugins/Editor/GameplayTagsEditor/Source/GameplayTagsEditor/Private/SGameplayTagContainerGraphPin.cpp
--- a/Engine/Plugins/Editor/GameplayTagsEditor/Source/GameplayTagsEditor/Private/SGameplayTagContainerGraphPin.cpp
+++ b/Engine/Plugins/Editor/GameplayTagsEditor/Source/GameplayTagsEditor/Private/SGameplayTagContainerGraphPin.cpp
@@ -4,6 +4,7 @@
 #include "Widgets/Input/SComboButton.h"
 #include "GameplayTagsModule.h"
 #include "Widgets/Layout/SScaleBox.h"
+#include "GameplayTagPinUtilities.h"
 
 #define LOCTEXT_NAMESPACE "GameplayTagGraphPin"
 
@@ -43,6 +44,8 @@ void SGameplayTagContainerGraphPin::ParseDefaultValueData()
 {
 	FString TagString = GraphPinObj->GetDefaultAsString();
 
+	FilterString = GameplayTagPinUtilities::ExtractTagFilterStringFromGraphPin(GraphPinObj);
+
 	if (TagString.StartsWith(TEXT("(")) && TagString.EndsWith(TEXT(")")))
 	{
 		TagString = TagString.LeftChop(1);
@@ -108,6 +111,7 @@ TSharedRef<SWidget> SGameplayTagContainerGraphPin::GetListContent()
 			.OnTagChanged(this, &SGameplayTagContainerGraphPin::RefreshTagList)
 			.TagContainerName( TEXT("SGameplayTagContainerGraphPin") )
 			.Visibility( this, &SGraphPin::GetDefaultValueVisibility )
+			.Filter(FilterString)
 		];
 }
 
diff --git a/Engine/Plugins/Editor/GameplayTagsEditor/Source/GameplayTagsEditor/Private/SGameplayTagContainerGraphPin.h b/Engine/Plugins/Editor/GameplayTagsEditor/Source/GameplayTagsEditor/Private/SGameplayTagContainerGraphPin.h
--- a/Engine/Plugins/Editor/GameplayTagsEditor/Source/GameplayTagsEditor/Private/SGameplayTagContainerGraphPin.h
+++ b/Engine/Plugins/Editor/GameplayTagsEditor/Source/GameplayTagsEditor/Private/SGameplayTagContainerGraphPin.h
@@ -63,4 +63,7 @@ private:
 
 	// The List View used to display the read only tag names on the node.
 	TSharedPtr<SListView<TSharedPtr<FString>>> TagListView;
+
+	// Tag filter string for the container
+	FString FilterString;
 };
diff --git a/Engine/Plugins/Editor/GameplayTagsEditor/Source/GameplayTagsEditor/Private/SGameplayTagGraphPin.cpp b/Engine/Plugins/Editor/GameplayTagsEditor/Source/GameplayTagsEditor/Private/SGameplayTagGraphPin.cpp
--- a/Engine/Plugins/Editor/GameplayTagsEditor/Source/GameplayTagsEditor/Private/SGameplayTagGraphPin.cpp
+++ b/Engine/Plugins/Editor/GameplayTagsEditor/Source/GameplayTagsEditor/Private/SGameplayTagGraphPin.cpp
@@ -2,10 +2,9 @@
 
 #include "SGameplayTagGraphPin.h"
 #include "Widgets/Input/SComboButton.h"
-#include "K2Node_CallFunction.h"
 #include "GameplayTagsModule.h"
 #include "Widgets/Layout/SScaleBox.h"
-#include "K2Node_VariableSet.h"
+#include "GameplayTagPinUtilities.h"
 
 #define LOCTEXT_NAMESPACE "GameplayTagGraphPin"
 
@@ -46,26 +45,7 @@ void SGameplayTagGraphPin::ParseDefaultValueData()
 {
 	FString TagString = GraphPinObj->GetDefaultAsString();
 
-	FilterString.Empty();
-	if (UScriptStruct* PinStructType = Cast<UScriptStruct>(GraphPinObj->PinType.PinSubCategoryObject.Get()))
-	{
-		FilterString = UGameplayTagsManager::Get().GetCategoriesMetaFromField(PinStructType);
-	}
-
-	if (FilterString.IsEmpty())
-	{
-		if (UK2Node_CallFunction* CallFuncNode = Cast<UK2Node_CallFunction>(GraphPinObj->GetOwningNode()))
-		{
-			if (UFunction* ThisFunction = CallFuncNode->GetTargetFunction())
-			{
-				FilterString = UGameplayTagsManager::Get().GetCategoriesMetaFromFunction(ThisFunction);
-			}
-		}
-		else if (UK2Node_VariableSet* ThisVariable = Cast<UK2Node_VariableSet>(GraphPinObj->GetOwningNode()))
-		{
-			FilterString = UGameplayTagsManager::Get().GetCategoriesMetaFromField(ThisVariable->GetPropertyForVariable());
-		}
-	}
+	FilterString = GameplayTagPinUtilities::ExtractTagFilterStringFromGraphPin(GraphPinObj);
 
 	if (TagString.StartsWith(TEXT("(")) && TagString.EndsWith(TEXT(")")))
 	{
diff --git a/Engine/Source/Runtime/GameplayTags/Classes/GameplayTagsManager.h b/Engine/Source/Runtime/GameplayTags/Classes/GameplayTagsManager.h
--- a/Engine/Source/Runtime/GameplayTags/Classes/GameplayTagsManager.h
+++ b/Engine/Source/Runtime/GameplayTags/Classes/GameplayTagsManager.h
@@ -548,7 +548,7 @@ class GAMEPLAYTAGS_API UGameplayTagsManager : public UObject
 	FString GetCategoriesMetaFromStruct(UScriptStruct* Struct) const { return GetCategoriesMetaFromField(Struct); }
 
 	/** Returns "GameplayTagFilter" meta property from given function, used for filtering by tag widget for any parameters of the function that end up as BP pins */
-	FString GetCategoriesMetaFromFunction(UFunction* Func) const;
+	FString GetCategoriesMetaFromFunction(UFunction* Func, FName ParamName = NAME_None) const;
 
 	/** Gets a list of all gameplay tag nodes added by the specific source */
 	void GetAllTagsFromSource(FName TagSource, TArray< TSharedPtr<FGameplayTagNode> >& OutTagArray) const;
diff --git a/Engine/Source/Runtime/GameplayTags/Private/GameplayTagsManager.cpp b/Engine/Source/Runtime/GameplayTags/Private/GameplayTagsManager.cpp
--- a/Engine/Source/Runtime/GameplayTags/Private/GameplayTagsManager.cpp
+++ b/Engine/Source/Runtime/GameplayTags/Private/GameplayTagsManager.cpp
@@ -1345,13 +1345,28 @@ FString UGameplayTagsManager::StaticGetCategoriesMetaFromPropertyHandle(TSharedP
 	return Categories;
 }
 
-FString UGameplayTagsManager::GetCategoriesMetaFromFunction(UFunction* ThisFunction) const
+FString UGameplayTagsManager::GetCategoriesMetaFromFunction(UFunction* ThisFunction, FName ParamName /** = NAME_None */) const
 {
 	FString FilterString;
-	if (ThisFunction->HasMetaData(NAME_GameplayTagFilter))
+	if (ThisFunction)
 	{
-		FilterString = ThisFunction->GetMetaData(NAME_GameplayTagFilter);
+		// If a param name was specified, check it first for UPARAM metadata
+		if (!ParamName.IsNone())
+		{
+			UProperty* ParamProp = FindField<UProperty>(ThisFunction, ParamName);
+			if (ParamProp)
+			{
+				FilterString = UGameplayTagsManager::Get().GetCategoriesMetaFromField(ParamProp);
+			}
+		}
+
+		// No filter found so far, fall back to UFUNCTION-level
+		if (FilterString.IsEmpty() && ThisFunction->HasMetaData(NAME_GameplayTagFilter))
+		{
+			FilterString = ThisFunction->GetMetaData(NAME_GameplayTagFilter);
+		}
 	}
+
 	return FilterString;
 }
 
