Fixed crash in GameplayTagQueryCustomization when a GameplayTagQuery that is a child property is modified

diff --git a/Engine/Plugins/Editor/GameplayTagsEditor/Source/GameplayTagsEditor/Private/GameplayTagQueryCustomization.cpp b/Engine/Plugins/Editor/GameplayTagsEditor/Source/GameplayTagsEditor/Private/GameplayTagQueryCustomization.cpp
--- a/Engine/Plugins/Editor/GameplayTagsEditor/Source/GameplayTagsEditor/Private/GameplayTagQueryCustomization.cpp
+++ b/Engine/Plugins/Editor/GameplayTagsEditor/Source/GameplayTagsEditor/Private/GameplayTagQueryCustomization.cpp
@@ -247,9 +247,15 @@ void FGameplayTagQueryCustomization::CloseWidgetWindow(bool WasCancelled)
 		if (StructPropertyHandle.IsValid())
 		{
 			UProperty* TheProperty = StructPropertyHandle->GetProperty();
+			TSharedPtr<IPropertyHandle> HeadPropertyHandle = StructPropertyHandle;
+
+			while (HeadPropertyHandle->GetParentHandle() && (HeadPropertyHandle->GetParentHandle()->GetProperty()))
+			{
+				HeadPropertyHandle = HeadPropertyHandle->GetParentHandle();
+			}
 
 			FEditPropertyChain PropertyChain;
-			PropertyChain.AddHead(TheProperty);
+			PropertyChain.AddHead(HeadPropertyHandle->GetProperty());
 			PropertyChain.SetActivePropertyNode(TheProperty);
 
 			FPropertyChangedEvent ChangeEvent(StructPropertyHandle->GetProperty(), EPropertyChangeType::ValueSet, nullptr);
