Restricted tag quality of life improvements.
Removed the default option for restricted gameplay tags.
No longer resets the source when adding several restricted tags in a row.

diff --git a/Engine/Plugins/Editor/GameplayTagsEditor/Source/GameplayTagsEditor/Private/SAddNewRestrictedGameplayTagWidget.cpp b/Engine/Plugins/Editor/GameplayTagsEditor/Source/GameplayTagsEditor/Private/SAddNewRestrictedGameplayTagWidget.cpp
--- a/Engine/Plugins/Editor/GameplayTagsEditor/Source/GameplayTagsEditor/Private/SAddNewRestrictedGameplayTagWidget.cpp
+++ b/Engine/Plugins/Editor/GameplayTagsEditor/Source/GameplayTagsEditor/Private/SAddNewRestrictedGameplayTagWidget.cpp
@@ -183,27 +183,31 @@ void SAddNewRestrictedGameplayTagWidget::PopulateTagSources()
 	UGameplayTagsManager& Manager = UGameplayTagsManager::Get();
 	RestrictedTagSources.Empty();
 
-	FName DefaultSource = FGameplayTagSource::GetDefaultName();
-
-	// Always ensure that the default source is first
-	RestrictedTagSources.Add( MakeShareable( new FName( DefaultSource ) ) );
-
 	TArray<const FGameplayTagSource*> Sources;
 	Manager.GetRestrictedTagSources(Sources);
 
+	// Used to make sure we have a non-empty list of restricted tag sources. Not an actual source.
+	FName PlaceholderSource = NAME_None;
+
+	// Add the placeholder source if no other sources exist
+	if (Sources.Num() == 0)
+	{
+		RestrictedTagSources.Add(MakeShareable(new FName(PlaceholderSource)));
+	}
+
 	for (const FGameplayTagSource* Source : Sources)
 	{
-		if (Source != nullptr && Source->SourceName != DefaultSource)
+		if (Source != nullptr && Source->SourceName != PlaceholderSource)
 		{
 			RestrictedTagSources.Add(MakeShareable(new FName(Source->SourceName)));
 		}
 	}
 }
 
-void SAddNewRestrictedGameplayTagWidget::Reset()
+void SAddNewRestrictedGameplayTagWidget::Reset(FName TagSource)
 {
 	SetTagName();
-	SelectTagSource();
+	SelectTagSource(TagSource);
 	SetAllowNonRestrictedChildren();
 	TagCommentTextBox->SetText(FText());
 }
@@ -274,6 +278,18 @@ void SAddNewRestrictedGameplayTagWidget::ValidateNewRestrictedTag()
 	bool bAllowNonRestrictedChildren = AllowNonRestrictedChildrenCheckBox->IsChecked();
 	FName TagSource = *TagSourcesComboBox->GetSelectedItem().Get();
 
+	if (TagSource == NAME_None)
+	{
+		FNotificationInfo Info(LOCTEXT("NoRestrictedSource", "You must specify a source file for restricted gameplay tags."));
+		Info.ExpireDuration = 10.f;
+		Info.bUseSuccessFailIcons = true;
+		Info.Image = FEditorStyle::GetBrush(TEXT("MessageLog.Error"));
+
+		AddRestrictedGameplayTagDialog = FSlateNotificationManager::Get().AddNotification(Info);
+
+		return;
+	}
+
 	TArray<FString> TagSourceOwners;
 	Manager.GetOwnersForTagSource(TagSource.ToString(), TagSourceOwners);
 
@@ -360,7 +376,7 @@ void SAddNewRestrictedGameplayTagWidget::CreateNewRestrictedGameplayTag()
 
 	OnRestrictedGameplayTagAdded.ExecuteIfBound(TagName, TagComment, TagSource);
 
-	Reset();
+	Reset(TagSource);
 }
 
 void SAddNewRestrictedGameplayTagWidget::CancelNewTag()
diff --git a/Engine/Plugins/Editor/GameplayTagsEditor/Source/GameplayTagsEditor/Private/SAddNewRestrictedGameplayTagWidget.h b/Engine/Plugins/Editor/GameplayTagsEditor/Source/GameplayTagsEditor/Private/SAddNewRestrictedGameplayTagWidget.h
--- a/Engine/Plugins/Editor/GameplayTagsEditor/Source/GameplayTagsEditor/Private/SAddNewRestrictedGameplayTagWidget.h
+++ b/Engine/Plugins/Editor/GameplayTagsEditor/Source/GameplayTagsEditor/Private/SAddNewRestrictedGameplayTagWidget.h
@@ -43,7 +43,7 @@ public:
 	void AddSubtagFromParent(const FString& ParentTagName, const FName& ParentTagSource, bool bAllowNonRestrictedChildren);
 
 	/** Resets all input fields */
-	void Reset();
+	void Reset(FName TagSource = NAME_None);
 
 private:
 
