If the gameplay tag widget has valid filter text we should continue to use it whenever we refresh the tag tree. This fixes bugs that could occur when you added or deleted a tag while you had filter text.

diff --git a/Engine/Plugins/Editor/GameplayTagsEditor/Source/GameplayTagsEditor/Private/SGameplayTagWidget.cpp b/Engine/Plugins/Editor/GameplayTagsEditor/Source/GameplayTagsEditor/Private/SGameplayTagWidget.cpp
--- a/Engine/Plugins/Editor/GameplayTagsEditor/Source/GameplayTagsEditor/Private/SGameplayTagWidget.cpp
+++ b/Engine/Plugins/Editor/GameplayTagsEditor/Source/GameplayTagsEditor/Private/SGameplayTagWidget.cpp
@@ -314,36 +314,41 @@ void SGameplayTagWidget::OnFilterTextChanged( const FText& InFilterText )
 {
 	FilterString = InFilterText.ToString();	
 
-	if( FilterString.IsEmpty() )
+	FilterTagTree();
+}
+
+void SGameplayTagWidget::FilterTagTree()
+{
+	if (FilterString.IsEmpty())
 	{
-		TagTreeWidget->SetTreeItemsSource( &TagItems );
+		TagTreeWidget->SetTreeItemsSource(&TagItems);
 
-		for( int32 iItem = 0; iItem < TagItems.Num(); ++iItem )
+		for (int32 iItem = 0; iItem < TagItems.Num(); ++iItem)
 		{
-			SetDefaultTagNodeItemExpansion( TagItems[iItem] );
+			SetDefaultTagNodeItemExpansion(TagItems[iItem]);
 		}
 	}
 	else
 	{
 		FilteredTagItems.Empty();
 
-		for( int32 iItem = 0; iItem < TagItems.Num(); ++iItem )
+		for (int32 iItem = 0; iItem < TagItems.Num(); ++iItem)
 		{
-			if( FilterChildrenCheck( TagItems[iItem] ) )
+			if (FilterChildrenCheck(TagItems[iItem]))
 			{
-				FilteredTagItems.Add( TagItems[iItem] );
-				SetTagNodeItemExpansion( TagItems[iItem], true );
+				FilteredTagItems.Add(TagItems[iItem]);
+				SetTagNodeItemExpansion(TagItems[iItem], true);
 			}
 			else
 			{
-				SetTagNodeItemExpansion( TagItems[iItem], false );
+				SetTagNodeItemExpansion(TagItems[iItem], false);
 			}
 		}
 
-		TagTreeWidget->SetTreeItemsSource( &FilteredTagItems );	
+		TagTreeWidget->SetTreeItemsSource(&FilteredTagItems);
 	}
-		
-	TagTreeWidget->RequestTreeRefresh();	
+
+	TagTreeWidget->RequestTreeRefresh();
 }
 
 bool SGameplayTagWidget::FilterChildrenCheck( TSharedPtr<FGameplayTagNode> InItem )
@@ -1212,7 +1217,7 @@ void SGameplayTagWidget::RefreshTags()
 		}
 	}
 
-	TagTreeWidget->SetTreeItemsSource(&TagItems);
+	FilterTagTree();
 }
 
 EVisibility SGameplayTagWidget::DetermineExpandableUIVisibility() const
diff --git a/Engine/Plugins/Editor/GameplayTagsEditor/Source/GameplayTagsEditor/Private/SGameplayTagWidget.h b/Engine/Plugins/Editor/GameplayTagsEditor/Source/GameplayTagsEditor/Private/SGameplayTagWidget.h
--- a/Engine/Plugins/Editor/GameplayTagsEditor/Source/GameplayTagsEditor/Private/SGameplayTagWidget.h
+++ b/Engine/Plugins/Editor/GameplayTagsEditor/Source/GameplayTagsEditor/Private/SGameplayTagWidget.h
@@ -103,6 +103,9 @@ public:
 
 private:
 
+	/* Filters the tree view based on the current filter text. */
+	void FilterTagTree();
+
 	/* string that sets the section of the ini file to use for this class*/ 
 	static const FString SettingsIniSection;
 
