Fix crash copying gameplay tag set node between blueprint graphs

diff --git a/Engine/Plugins/Editor/GameplayTagsEditor/Source/GameplayTagsEditor/Private/GameplayTagPinUtilities.cpp b/Engine/Plugins/Editor/GameplayTagsEditor/Source/GameplayTagsEditor/Private/GameplayTagPinUtilities.cpp
--- a/Engine/Plugins/Editor/GameplayTagsEditor/Source/GameplayTagsEditor/Private/GameplayTagPinUtilities.cpp
+++ b/Engine/Plugins/Editor/GameplayTagsEditor/Source/GameplayTagsEditor/Private/GameplayTagPinUtilities.cpp
@@ -29,7 +29,10 @@ FString GameplayTagPinUtilities::ExtractTagFilterStringFromGraphPin(UEdGraphPin*
 			}
 			else if (UK2Node_VariableSet* ThisVariable = Cast<UK2Node_VariableSet>(InTagPin->GetOwningNode()))
 			{
-				FilterString = TagManager.GetCategoriesMetaFromField(ThisVariable->GetPropertyForVariable());
+				if (UProperty* Property = ThisVariable->GetPropertyForVariable())
+				{
+					FilterString = TagManager.GetCategoriesMetaFromField(Property);
+				}
 			}
 		}
 	}
