- Fix issue wherein gameplay tag category filter wouldn't apply to function parameters in blueprints if they were part of a function terminator node
- Also constify some things that should have been const in the first place

diff --git a/Engine/Plugins/Editor/GameplayTagsEditor/Source/GameplayTagsEditor/Private/GameplayTagPinUtilities.cpp b/Engine/Plugins/Editor/GameplayTagsEditor/Source/GameplayTagsEditor/Private/GameplayTagPinUtilities.cpp
--- a/Engine/Plugins/Editor/GameplayTagsEditor/Source/GameplayTagsEditor/Private/GameplayTagPinUtilities.cpp
+++ b/Engine/Plugins/Editor/GameplayTagsEditor/Source/GameplayTagsEditor/Private/GameplayTagPinUtilities.cpp
@@ -5,6 +5,7 @@
 #include "SGraphPin.h"
 #include "K2Node_CallFunction.h"
 #include "K2Node_VariableSet.h"
+#include "K2Node_FunctionTerminator.h"
 
 FString GameplayTagPinUtilities::ExtractTagFilterStringFromGraphPin(UEdGraphPin* InTagPin)
 {
@@ -20,18 +21,23 @@ FString GameplayTagPinUtilities::ExtractTagFilterStringFromGraphPin(UEdGraphPin*
 
 		if (FilterString.IsEmpty())
 		{
-			if (UK2Node_CallFunction* CallFuncNode = Cast<UK2Node_CallFunction>(InTagPin->GetOwningNode()))
+			const UEdGraphNode* OwningNode = InTagPin->GetOwningNode();
+			if (const UK2Node_CallFunction* CallFuncNode = Cast<UK2Node_CallFunction>(OwningNode))
 			{
-				if (UFunction* ThisFunction = CallFuncNode->GetTargetFunction())
+				if (const UFunction* TargetFunction = CallFuncNode->GetTargetFunction())
 				{
-					FilterString = TagManager.GetCategoriesMetaFromFunction(ThisFunction, InTagPin->PinName);
+					FilterString = TagManager.GetCategoriesMetaFromFunction(TargetFunction, InTagPin->PinName);
 				}
 			}
-			else if (UK2Node_VariableSet* ThisVariable = Cast<UK2Node_VariableSet>(InTagPin->GetOwningNode()))
+			else if (const UK2Node_VariableSet* VariableSetNode = Cast<UK2Node_VariableSet>(OwningNode))
 			{
-				if (UProperty* Property = ThisVariable->GetPropertyForVariable())
+				FilterString = TagManager.GetCategoriesMetaFromField(VariableSetNode->GetPropertyForVariable());
+			}
+			else if (const UK2Node_FunctionTerminator* FuncTermNode = Cast<UK2Node_FunctionTerminator>(OwningNode))
+			{
+				if (const UFunction* SignatureFunction = FuncTermNode->FindSignatureFunction())
 				{
-					FilterString = TagManager.GetCategoriesMetaFromField(Property);
+					FilterString = TagManager.GetCategoriesMetaFromFunction(SignatureFunction, InTagPin->PinName);
 				}
 			}
 		}
diff --git a/Engine/Source/Runtime/GameplayTags/Classes/GameplayTagsManager.h b/Engine/Source/Runtime/GameplayTags/Classes/GameplayTagsManager.h
--- a/Engine/Source/Runtime/GameplayTags/Classes/GameplayTagsManager.h
+++ b/Engine/Source/Runtime/GameplayTags/Classes/GameplayTagsManager.h
@@ -548,7 +548,7 @@ class GAMEPLAYTAGS_API UGameplayTagsManager : public UObject
 	FString GetCategoriesMetaFromStruct(UScriptStruct* Struct) const { return GetCategoriesMetaFromField(Struct); }
 
 	/** Returns "GameplayTagFilter" meta property from given function, used for filtering by tag widget for any parameters of the function that end up as BP pins */
-	FString GetCategoriesMetaFromFunction(UFunction* Func, FName ParamName = NAME_None) const;
+	FString GetCategoriesMetaFromFunction(const UFunction* Func, FName ParamName = NAME_None) const;
 
 	/** Gets a list of all gameplay tag nodes added by the specific source */
 	void GetAllTagsFromSource(FName TagSource, TArray< TSharedPtr<FGameplayTagNode> >& OutTagArray) const;
diff --git a/Engine/Source/Runtime/GameplayTags/Private/GameplayTagsManager.cpp b/Engine/Source/Runtime/GameplayTags/Private/GameplayTagsManager.cpp
--- a/Engine/Source/Runtime/GameplayTags/Private/GameplayTagsManager.cpp
+++ b/Engine/Source/Runtime/GameplayTags/Private/GameplayTagsManager.cpp
@@ -1345,7 +1345,7 @@ FString UGameplayTagsManager::StaticGetCategoriesMetaFromPropertyHandle(TSharedP
 	return Categories;
 }
 
-FString UGameplayTagsManager::GetCategoriesMetaFromFunction(UFunction* ThisFunction, FName ParamName /** = NAME_None */) const
+FString UGameplayTagsManager::GetCategoriesMetaFromFunction(const UFunction* ThisFunction, FName ParamName /** = NAME_None */) const
 {
 	FString FilterString;
 	if (ThisFunction)
