Don't reset the tag source when we add a new tag in the editor.

diff --git a/Engine/Plugins/Editor/GameplayTagsEditor/Source/GameplayTagsEditor/Private/SAddNewGameplayTagWidget.cpp b/Engine/Plugins/Editor/GameplayTagsEditor/Source/GameplayTagsEditor/Private/SAddNewGameplayTagWidget.cpp
--- a/Engine/Plugins/Editor/GameplayTagsEditor/Source/GameplayTagsEditor/Private/SAddNewGameplayTagWidget.cpp
+++ b/Engine/Plugins/Editor/GameplayTagsEditor/Source/GameplayTagsEditor/Private/SAddNewGameplayTagWidget.cpp
@@ -144,7 +144,7 @@ void SAddNewGameplayTagWidget::Construct(const FArguments& InArgs)
 		]
 	];
 
-	Reset();
+	Reset(EResetType::ResetAll);
 }
 
 void SAddNewGameplayTagWidget::Tick( const FGeometry& AllottedGeometry, const double InCurrentTime, const float InDeltaTime )
@@ -178,10 +178,13 @@ void SAddNewGameplayTagWidget::PopulateTagSources()
 	}
 }
 
-void SAddNewGameplayTagWidget::Reset()
+void SAddNewGameplayTagWidget::Reset(EResetType ResetType)
 {
 	SetTagName();
-	SelectTagSource();
+	if (ResetType != EResetType::DoNotResetSource)
+	{
+		SelectTagSource();
+	}
 	TagCommentTextBox->SetText(FText());
 }
 
@@ -275,7 +278,7 @@ void SAddNewGameplayTagWidget::CreateNewGameplayTag()
 
 	OnGameplayTagAdded.ExecuteIfBound(TagName, TagComment, TagSource);
 
-	Reset();
+	Reset(EResetType::DoNotResetSource);
 }
 
 TSharedRef<SWidget> SAddNewGameplayTagWidget::OnGenerateTagSourcesComboBox(TSharedPtr<FName> InItem)
diff --git a/Engine/Plugins/Editor/GameplayTagsEditor/Source/GameplayTagsEditor/Private/SAddNewGameplayTagWidget.h b/Engine/Plugins/Editor/GameplayTagsEditor/Source/GameplayTagsEditor/Private/SAddNewGameplayTagWidget.h
--- a/Engine/Plugins/Editor/GameplayTagsEditor/Source/GameplayTagsEditor/Private/SAddNewGameplayTagWidget.h
+++ b/Engine/Plugins/Editor/GameplayTagsEditor/Source/GameplayTagsEditor/Private/SAddNewGameplayTagWidget.h
@@ -16,6 +16,12 @@ class GAMEPLAYTAGSEDITOR_VTABLE SAddNewGameplayTagWidget : public SCompoundWidge
 {
 public:
 
+	enum class EResetType : uint8
+	{
+		ResetAll,
+		DoNotResetSource
+	};
+
 	DECLARE_DELEGATE_ThreeParams(FOnGameplayTagAdded, const FString& /*TagName*/, const FString& /*TagComment*/, const FName& /*TagSource*/);
 
 	DECLARE_DELEGATE_RetVal_TwoParams(bool, FIsValidTag, const FString& /*TagName*/, FText* /*OutError*/)
@@ -44,7 +50,7 @@ public:
 	void AddSubtagFromParent(const FString& ParentTagName, const FName& ParentTagSource);
 
 	/** Resets all input fields */
-	void Reset();
+	void Reset(EResetType ResetType);
 
 private:
 
