Fixed gameplay tag categories filter not working for gameplay tags within structures

diff --git a/Engine/Plugins/Editor/GameplayTagsEditor/Source/GameplayTagsEditor/Private/GameplayTagPinUtilities.cpp b/Engine/Plugins/Editor/GameplayTagsEditor/Source/GameplayTagsEditor/Private/GameplayTagPinUtilities.cpp
--- a/Engine/Plugins/Editor/GameplayTagsEditor/Source/GameplayTagsEditor/Private/GameplayTagPinUtilities.cpp
+++ b/Engine/Plugins/Editor/GameplayTagsEditor/Source/GameplayTagsEditor/Private/GameplayTagPinUtilities.cpp
@@ -7,6 +7,8 @@
 #include "K2Node_VariableSet.h"
 #include "K2Node_FunctionTerminator.h"
 
+static FName NAME_Categories = FName("Categories");
+
 FString GameplayTagPinUtilities::ExtractTagFilterStringFromGraphPin(UEdGraphPin* InTagPin)
 {
 	FString FilterString;
@@ -19,9 +21,15 @@ FString GameplayTagPinUtilities::ExtractTagFilterStringFromGraphPin(UEdGraphPin*
 			FilterString = TagManager.GetCategoriesMetaFromField(PinStructType);
 		}
 
+		UEdGraphNode* OwningNode = InTagPin->GetOwningNode();
+
+		if (FilterString.IsEmpty())
+		{
+			FilterString = OwningNode->GetPinMetaData(InTagPin->PinName, NAME_Categories);
+		}
+
 		if (FilterString.IsEmpty())
 		{
-			const UEdGraphNode* OwningNode = InTagPin->GetOwningNode();
 			if (const UK2Node_CallFunction* CallFuncNode = Cast<UK2Node_CallFunction>(OwningNode))
 			{
 				if (const UFunction* TargetFunction = CallFuncNode->GetTargetFunction())
