Editor: Surfaced 'Search For References' directly in the context menu for FGameplayTag properties, individual tag rows within a FGameplayTagContainer property, and FGameplayTagContainer properties (lists referencers to any of the tags in the container)

diff --git a/Engine/Plugins/Editor/GameplayTagsEditor/Source/GameplayTagsEditor/Private/GameplayTagContainerCustomization.cpp b/Engine/Plugins/Editor/GameplayTagsEditor/Source/GameplayTagsEditor/Private/GameplayTagContainerCustomization.cpp
--- a/Engine/Plugins/Editor/GameplayTagsEditor/Source/GameplayTagsEditor/Private/GameplayTagContainerCustomization.cpp
+++ b/Engine/Plugins/Editor/GameplayTagsEditor/Source/GameplayTagsEditor/Private/GameplayTagContainerCustomization.cpp
@@ -24,6 +24,8 @@ void FGameplayTagContainerCustomization::CustomizeHeader(TSharedRef<class IPrope
 
 	BuildEditableContainerList();
 
+	FUIAction SearchForReferencesAction(FExecuteAction::CreateSP(this, &FGameplayTagContainerCustomization::OnWholeContainerSearchForReferences));
+
 	HeaderRow
 		.NameContent()
 		[
@@ -72,7 +74,11 @@ void FGameplayTagContainerCustomization::CustomizeHeader(TSharedRef<class IPrope
 					ActiveTags()
 				]
 			]
-		];
+		]
+		.AddCustomContextMenuAction(SearchForReferencesAction,
+			LOCTEXT("WholeContainerSearchForReferences", "Search For References"),
+			LOCTEXT("WholeContainerSearchForReferencesTooltip", "Find referencers that reference *any* of the tags in this container"),
+			FSlateIcon());
 
 	GEditor->RegisterForUndo(this);
 }
@@ -134,35 +140,96 @@ TSharedRef<ITableRow> FGameplayTagContainerCustomization::MakeListViewWidget(TSh
 
 	return SNew( STableRow< TSharedPtr<FString> >, OwnerTable )
 	[
-		SNew(SHorizontalBox)
-
-		+ SHorizontalBox::Slot()
-		.AutoWidth()
-		.VAlign(VAlign_Center)
-		.Padding(0,0,2,0)
+		SNew(SBorder)
+		.OnMouseButtonDown(this, &FGameplayTagContainerCustomization::OnSingleTagMouseButtonPressed, *Item.Get())
+		.Padding(0.0f)
+		.BorderImage(FEditorStyle::GetBrush("NoBorder"))
 		[
-			SNew(SButton)
-			.IsEnabled(!StructPropertyHandle->IsEditConst())
-			.ContentPadding(FMargin(0))
-			.ButtonStyle(FEditorStyle::Get(), "FlatButton.Danger")
-			.ForegroundColor(FSlateColor::UseForeground())
-			.OnClicked(this, &FGameplayTagContainerCustomization::OnRemoveTagClicked, *Item.Get())
+			SNew(SHorizontalBox)
+			+ SHorizontalBox::Slot()
+			.AutoWidth()
+			.VAlign(VAlign_Center)
+			.Padding(0,0,2,0)
 			[
-				SNew(STextBlock)
-				.Font(FEditorStyle::Get().GetFontStyle("FontAwesome.9"))
-				.Text(FEditorFontGlyphs::Times)
+				SNew(SButton)
+				.IsEnabled(!StructPropertyHandle->IsEditConst())
+				.ContentPadding(FMargin(0))
+				.ButtonStyle(FEditorStyle::Get(), "FlatButton.Danger")
+				.ForegroundColor(FSlateColor::UseForeground())
+				.OnClicked(this, &FGameplayTagContainerCustomization::OnRemoveTagClicked, *Item.Get())
+				[
+					SNew(STextBlock)
+					.Font(FEditorStyle::Get().GetFontStyle("FontAwesome.9"))
+					.Text(FEditorFontGlyphs::Times)
+				]
 			]
-		]
 
-		+ SHorizontalBox::Slot()
-		.AutoWidth()
-		.VAlign(VAlign_Center)
-		[
-			TagItem.ToSharedRef()
+			+ SHorizontalBox::Slot()
+			.AutoWidth()
+			.VAlign(VAlign_Center)
+			[
+				TagItem.ToSharedRef()
+			]
 		]
 	];
 }
 
+FReply FGameplayTagContainerCustomization::OnSingleTagMouseButtonPressed(const FGeometry& MyGeometry, const FPointerEvent& MouseEvent, FString TagName)
+{
+	if (MouseEvent.IsMouseButtonDown(EKeys::RightMouseButton))
+	{
+		FMenuBuilder MenuBuilder(/*bShouldCloseWindowAfterMenuSelection=*/ true, /*CommandList=*/ nullptr);
+
+		FUIAction SearchForReferencesAction(FExecuteAction::CreateSP(this, &FGameplayTagContainerCustomization::OnSingleTagSearchForReferences, TagName));
+
+		MenuBuilder.BeginSection(NAME_None, FText::Format(LOCTEXT("SingleTagMenuHeading", "Tag Actions ({0})"), FText::AsCultureInvariant(TagName)));
+		MenuBuilder.AddMenuEntry(
+			LOCTEXT("SingleTagSearchForReferences", "Search For References"),
+			FText::Format(LOCTEXT("SingleTagSearchForReferencesTooltip", "Find references to the tag {0}"), FText::AsCultureInvariant(TagName)),
+			FSlateIcon(),
+			SearchForReferencesAction);
+		MenuBuilder.EndSection();
+
+		// Spawn context menu
+		FWidgetPath WidgetPath = MouseEvent.GetEventPath() != nullptr ? *MouseEvent.GetEventPath() : FWidgetPath();
+		FSlateApplication::Get().PushMenu(TagListView.ToSharedRef(), WidgetPath, MenuBuilder.MakeWidget(), MouseEvent.GetScreenSpacePosition(), FPopupTransitionEffect(FPopupTransitionEffect::ContextMenu));
+
+		return FReply::Handled();
+	}
+
+	return FReply::Unhandled();
+}
+
+void FGameplayTagContainerCustomization::OnSingleTagSearchForReferences(FString TagName)
+{
+	FName TagFName(*TagName, FNAME_Find);
+	if (FEditorDelegates::OnOpenReferenceViewer.IsBound() && !TagFName.IsNone())
+	{
+		TArray<FAssetIdentifier> AssetIdentifiers;
+		AssetIdentifiers.Emplace(FGameplayTag::StaticStruct(), TagFName);
+		FEditorDelegates::OnOpenReferenceViewer.Broadcast(AssetIdentifiers, FReferenceViewerParams());
+	}
+}
+
+void FGameplayTagContainerCustomization::OnWholeContainerSearchForReferences()
+{
+	if (FEditorDelegates::OnOpenReferenceViewer.IsBound())
+	{
+		TArray<FAssetIdentifier> AssetIdentifiers;
+		AssetIdentifiers.Reserve(TagNames.Num());
+		for (TSharedPtr<FString>& TagNamePtr : TagNames)
+		{
+			FName TagFName(**TagNamePtr.Get(), FNAME_Find);
+			if (!TagFName.IsNone())
+			{
+				AssetIdentifiers.Emplace(FGameplayTag::StaticStruct(), TagFName);
+			}
+		}
+
+		FEditorDelegates::OnOpenReferenceViewer.Broadcast(AssetIdentifiers, FReferenceViewerParams());
+	}
+}
+
 void FGameplayTagContainerCustomization::OnTagDoubleClicked(FString TagName)
 {
 	UGameplayTagsManager::Get().NotifyGameplayTagDoubleClickedEditor(TagName);
diff --git a/Engine/Plugins/Editor/GameplayTagsEditor/Source/GameplayTagsEditor/Private/GameplayTagContainerCustomization.h b/Engine/Plugins/Editor/GameplayTagsEditor/Source/GameplayTagsEditor/Private/GameplayTagContainerCustomization.h
--- a/Engine/Plugins/Editor/GameplayTagsEditor/Source/GameplayTagsEditor/Private/GameplayTagContainerCustomization.h
+++ b/Engine/Plugins/Editor/GameplayTagsEditor/Source/GameplayTagsEditor/Private/GameplayTagContainerCustomization.h
@@ -81,5 +81,9 @@ private:
 
 	void OnTagDoubleClicked(FString TagName);
 	FReply OnRemoveTagClicked(FString TagName);
+
+	FReply OnSingleTagMouseButtonPressed(const FGeometry& MyGeometry, const FPointerEvent& MouseEvent, FString TagName);
+	void OnSingleTagSearchForReferences(FString TagName);
+	void OnWholeContainerSearchForReferences();
 };
 
diff --git a/Engine/Plugins/Editor/GameplayTagsEditor/Source/GameplayTagsEditor/Private/GameplayTagCustomization.cpp b/Engine/Plugins/Editor/GameplayTagsEditor/Source/GameplayTagsEditor/Private/GameplayTagCustomization.cpp
--- a/Engine/Plugins/Editor/GameplayTagsEditor/Source/GameplayTagsEditor/Private/GameplayTagCustomization.cpp
+++ b/Engine/Plugins/Editor/GameplayTagsEditor/Source/GameplayTagsEditor/Private/GameplayTagCustomization.cpp
@@ -26,6 +26,8 @@ void FGameplayTagCustomization::CustomizeHeader(TSharedRef<class IPropertyHandle
 
 	BuildEditableContainerList();
 
+	FUIAction SearchForReferencesAction(FExecuteAction::CreateSP(this, &FGameplayTagCustomization::OnSearchForReferences));
+
 	HeaderRow
 	.NameContent()
 	[
@@ -74,7 +76,11 @@ void FGameplayTagCustomization::CustomizeHeader(TSharedRef<class IPropertyHandle
 				.OnNavigate( this, &FGameplayTagCustomization::OnTagDoubleClicked)
 			]
 		]
-	];
+	]
+	.AddCustomContextMenuAction(SearchForReferencesAction,
+		LOCTEXT("FGameplayTagCustomization_SearchForReferences", "Search For References"),
+		LOCTEXT("FGameplayTagCustomization_SearchForReferencesTooltip", "Find references for this tag"),
+		FSlateIcon());
 
 	GEditor->RegisterForUndo(this);
 }
@@ -84,6 +90,17 @@ void FGameplayTagCustomization::OnTagDoubleClicked()
 	UGameplayTagsManager::Get().NotifyGameplayTagDoubleClickedEditor(TagName);
 }
 
+void FGameplayTagCustomization::OnSearchForReferences()
+{
+	FName TagFName(*TagName, FNAME_Find);
+	if (FEditorDelegates::OnOpenReferenceViewer.IsBound() && !TagFName.IsNone())
+	{
+		TArray<FAssetIdentifier> AssetIdentifiers;
+		AssetIdentifiers.Add(FAssetIdentifier(FGameplayTag::StaticStruct(), TagFName));
+		FEditorDelegates::OnOpenReferenceViewer.Broadcast(AssetIdentifiers, FReferenceViewerParams());
+	}
+}
+
 EVisibility FGameplayTagCustomization::GetVisibilityForTagTextBlockWidget(bool ForTextWidget) const
 {
 	return (UGameplayTagsManager::Get().ShowGameplayTagAsHyperLinkEditor(TagName) ^ ForTextWidget) ? EVisibility::Visible : EVisibility::Collapsed;
diff --git a/Engine/Plugins/Editor/GameplayTagsEditor/Source/GameplayTagsEditor/Private/GameplayTagCustomization.h b/Engine/Plugins/Editor/GameplayTagsEditor/Source/GameplayTagsEditor/Private/GameplayTagCustomization.h
--- a/Engine/Plugins/Editor/GameplayTagsEditor/Source/GameplayTagsEditor/Private/GameplayTagCustomization.h
+++ b/Engine/Plugins/Editor/GameplayTagsEditor/Source/GameplayTagsEditor/Private/GameplayTagCustomization.h
@@ -60,6 +60,8 @@ private:
 	FString TagName;
 
 	void OnTagDoubleClicked();
+	void OnSearchForReferences();
+
 	EVisibility GetVisibilityForTagTextBlockWidget(bool ForTextWidget) const;
 
 	TSharedPtr<class SComboButton> EditButton;
