Fixed crash when adding a gameplay tag without a valid tag source selection.

diff --git a/Engine/Plugins/Editor/GameplayTagsEditor/Source/GameplayTagsEditor/Private/SAddNewGameplayTagWidget.cpp b/Engine/Plugins/Editor/GameplayTagsEditor/Source/GameplayTagsEditor/Private/SAddNewGameplayTagWidget.cpp
--- a/Engine/Plugins/Editor/GameplayTagsEditor/Source/GameplayTagsEditor/Private/SAddNewGameplayTagWidget.cpp
+++ b/Engine/Plugins/Editor/GameplayTagsEditor/Source/GameplayTagsEditor/Private/SAddNewGameplayTagWidget.cpp
@@ -176,6 +176,12 @@ void SAddNewGameplayTagWidget::PopulateTagSources()
 			TagSources.Add(MakeShareable(new FName(Source->SourceName)));
 		}
 	}
+
+	//Set selection to the latest added source
+	if (TagSourcesComboBox != nullptr)
+	{
+		TagSourcesComboBox->SetSelectedItem(TagSources.Last());
+	}	
 }
 
 void SAddNewGameplayTagWidget::Reset(EResetType ResetType)
@@ -249,6 +255,11 @@ void SAddNewGameplayTagWidget::CreateNewGameplayTag()
 		return;
 	}
 
+	if (TagSourcesComboBox->GetSelectedItem().Get() == nullptr)
+	{
+		return;
+	}
+
 	FText TagNameAsText = TagNameTextBox->GetText();
 	FString TagName = TagNameAsText.ToString();
 	FString TagComment = TagCommentTextBox->GetText().ToString();
