Added a favorite icon next to the tag source in the add tag widget to let users change the default ini to save new tags into
- Also changed UGameplayTagsDeveloperSettings to EditorPerProjectUserSettings to avoid users accidentally checking in their DeveloperConfigName

diff --git a/Engine/Plugins/Editor/GameplayTagsEditor/Source/GameplayTagsEditor/Private/SAddNewGameplayTagWidget.cpp b/Engine/Plugins/Editor/GameplayTagsEditor/Source/GameplayTagsEditor/Private/SAddNewGameplayTagWidget.cpp
--- a/Engine/Plugins/Editor/GameplayTagsEditor/Source/GameplayTagsEditor/Private/SAddNewGameplayTagWidget.cpp
+++ b/Engine/Plugins/Editor/GameplayTagsEditor/Source/GameplayTagsEditor/Private/SAddNewGameplayTagWidget.cpp
@@ -124,6 +124,22 @@ void SAddNewGameplayTagWidget::Construct(const FArguments& InArgs)
 					.Font(IDetailLayoutBuilder::GetDetailFont())
 				]
 			]
+
+			+ SHorizontalBox::Slot()
+			.AutoWidth()
+			.VAlign(VAlign_Center)
+			[
+				SNew( SButton )
+				.ButtonStyle( FEditorStyle::Get(), "NoBorder" )
+				.Visibility(this, &SAddNewGameplayTagWidget::OnGetTagSourceFavoritesVisibility)
+				.OnClicked(this, &SAddNewGameplayTagWidget::OnToggleTagSourceFavoriteClicked)
+				.ToolTipText(LOCTEXT("ToggleFavoriteTooltip", "Toggle whether or not this tag source is your favorite source (new tags will go into your favorite source by default)"))
+				.ContentPadding(0)
+				[
+					SNew(SImage)
+					.Image(this, &SAddNewGameplayTagWidget::OnGetTagSourceFavoriteImage)
+				]
+			]
 		]
 
 		// Add Tag Button
@@ -147,6 +163,29 @@ void SAddNewGameplayTagWidget::Construct(const FArguments& InArgs)
 	Reset(EResetType::ResetAll);
 }
 
+EVisibility SAddNewGameplayTagWidget::OnGetTagSourceFavoritesVisibility() const
+{
+	return (TagSources.Num() > 1) ? EVisibility::Visible : EVisibility::Collapsed;
+}
+
+FReply SAddNewGameplayTagWidget::OnToggleTagSourceFavoriteClicked()
+{
+	const FName ActiveTagSource = *TagSourcesComboBox->GetSelectedItem().Get();
+	const bool bWasFavorite = FGameplayTagSource::GetFavoriteName() == ActiveTagSource;
+
+	FGameplayTagSource::SetFavoriteName(bWasFavorite ? NAME_None : ActiveTagSource);
+
+	return FReply::Handled();
+}
+
+const FSlateBrush* SAddNewGameplayTagWidget::OnGetTagSourceFavoriteImage() const
+{
+	const FName ActiveTagSource = *TagSourcesComboBox->GetSelectedItem().Get();
+	const bool bIsFavoriteTagSource = FGameplayTagSource::GetFavoriteName() == ActiveTagSource;
+
+	return FEditorStyle::GetBrush(bIsFavoriteTagSource ? TEXT("PropertyWindow.Favorites_Enabled") : TEXT("PropertyWindow.Favorites_Disabled"));
+}
+
 void SAddNewGameplayTagWidget::Tick( const FGeometry& AllottedGeometry, const double InCurrentTime, const float InDeltaTime )
 {
 	if (bShouldGetKeyboardFocus)
@@ -189,7 +228,7 @@ void SAddNewGameplayTagWidget::Reset(EResetType ResetType)
 	SetTagName();
 	if (ResetType != EResetType::DoNotResetSource)
 	{
-		SelectTagSource();
+		SelectTagSource(FGameplayTagSource::GetFavoriteName());
 	}
 	TagCommentTextBox->SetText(FText());
 }
diff --git a/Engine/Plugins/Editor/GameplayTagsEditor/Source/GameplayTagsEditor/Private/SAddNewGameplayTagWidget.h b/Engine/Plugins/Editor/GameplayTagsEditor/Source/GameplayTagsEditor/Private/SAddNewGameplayTagWidget.h
--- a/Engine/Plugins/Editor/GameplayTagsEditor/Source/GameplayTagsEditor/Private/SAddNewGameplayTagWidget.h
+++ b/Engine/Plugins/Editor/GameplayTagsEditor/Source/GameplayTagsEditor/Private/SAddNewGameplayTagWidget.h
@@ -78,6 +78,10 @@ private:
 	/** Creates the text displayed by the combo box when an option is selected */
 	FText CreateTagSourcesComboBoxContent() const;
 
+	EVisibility OnGetTagSourceFavoritesVisibility() const;
+	FReply OnToggleTagSourceFavoriteClicked();
+	const FSlateBrush* OnGetTagSourceFavoriteImage() const;
+
 private:
 
 	/** All potential INI files where a gameplay tag can be stored */
diff --git a/Engine/Source/Runtime/GameplayTags/Classes/GameplayTagsManager.h b/Engine/Source/Runtime/GameplayTags/Classes/GameplayTagsManager.h
--- a/Engine/Source/Runtime/GameplayTags/Classes/GameplayTagsManager.h
+++ b/Engine/Source/Runtime/GameplayTags/Classes/GameplayTagsManager.h
@@ -84,7 +84,7 @@ enum class EGameplayTagSelectionType : uint8
 
 /** Struct defining where gameplay tags are loaded/saved from. Mostly for the editor */
 USTRUCT()
-struct FGameplayTagSource
+struct GAMEPLAYTAGS_API FGameplayTagSource
 {
 	GENERATED_USTRUCT_BODY()
 
@@ -114,24 +114,16 @@ struct FGameplayTagSource
 	{
 	}
 
-	static FName GetNativeName()
-	{
-		static FName NativeName = FName(TEXT("Native"));
-		return NativeName;
-	}
+	static FName GetNativeName();
 
-	static FName GetDefaultName()
-	{
-		static FName DefaultName = FName(TEXT("DefaultGameplayTags.ini"));
-		return DefaultName;
-	}
+	static FName GetDefaultName();
 
 #if WITH_EDITOR
-	static FName GetTransientEditorName()
-	{
-		static FName TransientEditorName = FName(TEXT("TransientEditor"));
-		return TransientEditorName;
-	}
+	static FName GetFavoriteName();
+
+	static void SetFavoriteName(FName TagSourceToFavorite);
+
+	static FName GetTransientEditorName();
 #endif
 };
 
diff --git a/Engine/Source/Runtime/GameplayTags/Classes/GameplayTagsSettings.h b/Engine/Source/Runtime/GameplayTags/Classes/GameplayTagsSettings.h
--- a/Engine/Source/Runtime/GameplayTags/Classes/GameplayTagsSettings.h
+++ b/Engine/Source/Runtime/GameplayTags/Classes/GameplayTagsSettings.h
@@ -7,6 +7,7 @@
 #include "UObject/Object.h"
 #include "UObject/SoftObjectPath.h"
 #include "GameplayTagsManager.h"
+#include "Engine/DeveloperSettings.h"
 #include "GameplayTagsSettings.generated.h"
 
 /** A single redirect from a deleted tag to the new tag that should replace it */
@@ -186,12 +187,18 @@ private:
 #endif
 };
 
-UCLASS(config=GameplayTags, notplaceable)
-class GAMEPLAYTAGS_API UGameplayTagsDeveloperSettings : public UObject
+UCLASS(config=EditorPerProjectUserSettings, meta=(DisplayName="Gameplay Tag Editing"))
+class GAMEPLAYTAGS_API UGameplayTagsDeveloperSettings : public UDeveloperSettings
 {
 	GENERATED_UCLASS_BODY()
 
+	virtual FName GetCategoryName() const override;
+
 	/** Allows new tags to be saved into their own INI file. This is make merging easier for non technical developers by setting up their own ini file. */
 	UPROPERTY(config, EditAnywhere, Category=GameplayTags)
 	FString DeveloperConfigName;
+
+	/** Stores the favorite tag source, used as the default ini when adding new tags, can be toggled on/off using the button next to the tag source picker */
+	UPROPERTY(config, EditAnywhere, Category=GameplayTags)
+	FName FavoriteTagSource;
 };
diff --git a/Engine/Source/Runtime/GameplayTags/Private/GameplayTagsManager.cpp b/Engine/Source/Runtime/GameplayTags/Private/GameplayTagsManager.cpp
--- a/Engine/Source/Runtime/GameplayTags/Private/GameplayTagsManager.cpp
+++ b/Engine/Source/Runtime/GameplayTags/Private/GameplayTagsManager.cpp
@@ -34,6 +34,55 @@ const FName UGameplayTagsManager::NAME_GameplayTagFilter("GameplayTagFilter");
 
 #define LOCTEXT_NAMESPACE "GameplayTagManager"
 
+//////////////////////////////////////////////////////////////////////
+// FGameplayTagSource
+
+static const FName NAME_Native = FName(TEXT("Native"));
+static const FName NAME_DefaultGameplayTagsIni("DefaultGameplayTags.ini");
+
+FName FGameplayTagSource::GetNativeName()
+{
+	return NAME_Native;
+}
+
+FName FGameplayTagSource::GetDefaultName()
+{
+	return NAME_DefaultGameplayTagsIni;
+}
+
+#if WITH_EDITOR
+static const FName NAME_TransientEditor("TransientEditor");
+
+FName FGameplayTagSource::GetFavoriteName()
+{
+	return GetDefault<UGameplayTagsDeveloperSettings>()->FavoriteTagSource;
+}
+
+void FGameplayTagSource::SetFavoriteName(FName TagSourceToFavorite)
+{
+	UGameplayTagsDeveloperSettings* MutableSettings = GetMutableDefault<UGameplayTagsDeveloperSettings>();
+
+	if (MutableSettings->FavoriteTagSource != TagSourceToFavorite)
+	{
+		MutableSettings->Modify();
+		MutableSettings->FavoriteTagSource = TagSourceToFavorite;
+
+		FPropertyChangedEvent ChangeEvent(MutableSettings->GetClass()->FindPropertyByName(GET_MEMBER_NAME_CHECKED(UGameplayTagsDeveloperSettings, FavoriteTagSource)), EPropertyChangeType::ValueSet);
+		MutableSettings->PostEditChangeProperty(ChangeEvent);
+		
+		MutableSettings->SaveConfig();
+	}
+}
+
+FName FGameplayTagSource::GetTransientEditorName()
+{
+	return NAME_TransientEditor;
+}
+#endif
+
+//////////////////////////////////////////////////////////////////////
+// UGameplayTagsManager
+
 UGameplayTagsManager* UGameplayTagsManager::SingletonManager = nullptr;
 
 UGameplayTagsManager::UGameplayTagsManager(const FObjectInitializer& ObjectInitializer)
diff --git a/Engine/Source/Runtime/GameplayTags/Private/GameplayTagsSettings.cpp b/Engine/Source/Runtime/GameplayTags/Private/GameplayTagsSettings.cpp
--- a/Engine/Source/Runtime/GameplayTags/Private/GameplayTagsSettings.cpp
+++ b/Engine/Source/Runtime/GameplayTags/Private/GameplayTagsSettings.cpp
@@ -137,3 +137,10 @@ UGameplayTagsDeveloperSettings::UGameplayTagsDeveloperSettings(const FObjectInit
 {
 	
 }
+
+static const FName NAME_Advanced("Advanced");
+
+FName UGameplayTagsDeveloperSettings::GetCategoryName() const
+{
+	return NAME_Advanced;
+}
