Added missing SImage.h include
#rnx

diff --git a/Engine/Plugins/Editor/GameplayTagsEditor/Source/GameplayTagsEditor/Private/SAddNewGameplayTagWidget.cpp b/Engine/Plugins/Editor/GameplayTagsEditor/Source/GameplayTagsEditor/Private/SAddNewGameplayTagWidget.cpp
--- a/Engine/Plugins/Editor/GameplayTagsEditor/Source/GameplayTagsEditor/Private/SAddNewGameplayTagWidget.cpp
+++ b/Engine/Plugins/Editor/GameplayTagsEditor/Source/GameplayTagsEditor/Private/SAddNewGameplayTagWidget.cpp
@@ -7,6 +7,7 @@
 #include "GameplayTagsModule.h"
 #include "Widgets/Input/SButton.h"
 #include "Misc/MessageDialog.h"
+#include "Widgets/Images/SImage.h"
 
 #define LOCTEXT_NAMESPACE "AddNewGameplayTagWidget"
 
