Don't show irrelevant options in gameplay tag widget when we don't have tag containers to modify.

diff --git a/Engine/Plugins/Editor/GameplayTagsEditor/Source/GameplayTagsEditor/Private/SGameplayTagWidget.cpp b/Engine/Plugins/Editor/GameplayTagsEditor/Source/GameplayTagsEditor/Private/SGameplayTagWidget.cpp
--- a/Engine/Plugins/Editor/GameplayTagsEditor/Source/GameplayTagsEditor/Private/SGameplayTagWidget.cpp
+++ b/Engine/Plugins/Editor/GameplayTagsEditor/Source/GameplayTagsEditor/Private/SGameplayTagWidget.cpp
@@ -926,15 +926,20 @@ TSharedRef<SWidget> SGameplayTagWidget::MakeTagActionsMenu(TSharedPtr<FGameplayT
 		MenuBuilder.AddMenuEntry(LOCTEXT("GameplayTagWidget_DeleteTag", "Delete"), LOCTEXT("GameplayTagWidget_DeleteTagTooltip", "Delete this tag"), FSlateIcon(), FUIAction(DeleteAction));
 	}
 
-	if (IsExactTagInCollection(InTagNode))
+	// Only include these menu items if we have tag containers to modify
+	if (TagContainers.Num() > 0)
 	{
-		FExecuteAction RemoveAction = FExecuteAction::CreateSP(this, &SGameplayTagWidget::OnRemoveTag, InTagNode);
-		MenuBuilder.AddMenuEntry(LOCTEXT("GameplayTagWidget_RemoveTag", "Remove Exact Tag"), LOCTEXT("GameplayTagWidget_RemoveTagTooltip", "Remove this exact tag, Parent and Child Tags will not be effected."), FSlateIcon(), FUIAction(RemoveAction));
-	}
-	else
-	{
-		FExecuteAction AddAction = FExecuteAction::CreateSP(this, &SGameplayTagWidget::OnAddTag, InTagNode);
-		MenuBuilder.AddMenuEntry(LOCTEXT("GameplayTagWidget_AddTag", "Add Exact Tag"), LOCTEXT("GameplayTagWidget_AddTagTooltip", "Add this exact tag, Parent and Child Child Tags will not be effected."), FSlateIcon(), FUIAction(AddAction));
+		// Either Remove or Add Exact Tag depending on if we have the exact tag or not
+		if (IsExactTagInCollection(InTagNode))
+		{
+			FExecuteAction RemoveAction = FExecuteAction::CreateSP(this, &SGameplayTagWidget::OnRemoveTag, InTagNode);
+			MenuBuilder.AddMenuEntry(LOCTEXT("GameplayTagWidget_RemoveTag", "Remove Exact Tag"), LOCTEXT("GameplayTagWidget_RemoveTagTooltip", "Remove this exact tag, Parent and Child Tags will not be effected."), FSlateIcon(), FUIAction(RemoveAction));
+		}
+		else
+		{
+			FExecuteAction AddAction = FExecuteAction::CreateSP(this, &SGameplayTagWidget::OnAddTag, InTagNode);
+			MenuBuilder.AddMenuEntry(LOCTEXT("GameplayTagWidget_AddTag", "Add Exact Tag"), LOCTEXT("GameplayTagWidget_AddTagTooltip", "Add this exact tag, Parent and Child Child Tags will not be effected."), FSlateIcon(), FUIAction(AddAction));
+		}
 	}
 
 	// Search for References
