Editor: Added a tooltip to gameplay tag queries to show the full query string (and fixed some whitespace problems in the declarative syntax for the widget)

diff --git a/Engine/Plugins/Editor/GameplayTagsEditor/Source/GameplayTagsEditor/Private/GameplayTagQueryCustomization.cpp b/Engine/Plugins/Editor/GameplayTagsEditor/Source/GameplayTagsEditor/Private/GameplayTagQueryCustomization.cpp
--- a/Engine/Plugins/Editor/GameplayTagsEditor/Source/GameplayTagsEditor/Private/GameplayTagQueryCustomization.cpp
+++ b/Engine/Plugins/Editor/GameplayTagsEditor/Source/GameplayTagsEditor/Private/GameplayTagQueryCustomization.cpp
@@ -35,44 +35,45 @@ void FGameplayTagQueryCustomization::CustomizeHeader(TSharedRef<class IPropertyH
 		[
 			StructPropertyHandle->CreatePropertyNameWidget()
 		]
-	.ValueContent()
+		.ValueContent()
 		.MaxDesiredWidth(512)
 		[
 			SNew(SHorizontalBox)
-		+SHorizontalBox::Slot()
-		.AutoWidth()
-		.VAlign(VAlign_Center)
-		[
-			SNew(SVerticalBox)
-			+SVerticalBox::Slot()
-		.AutoHeight()
-		[
-			SNew(SButton)
-			.Text(this, &FGameplayTagQueryCustomization::GetEditButtonText)
-		.OnClicked(this, &FGameplayTagQueryCustomization::OnEditButtonClicked)
-		]
-			+SVerticalBox::Slot()
-		.AutoHeight()
-		[
-			SNew(SButton)
-			.IsEnabled(!bReadOnly)
-		.Text(LOCTEXT("GameplayTagQueryCustomization_Clear", "Clear All"))
-		.OnClicked(this, &FGameplayTagQueryCustomization::OnClearAllButtonClicked)
-		.Visibility(this, &FGameplayTagQueryCustomization::GetClearAllVisibility)
-		]
-		]
-		+SHorizontalBox::Slot()
-		.AutoWidth()
-		[
-			SNew(SBorder)
-			.Padding(4.0f)
-		.Visibility(this, &FGameplayTagQueryCustomization::GetQueryDescVisibility)
-		[
-			SNew(STextBlock)
-			.Text(this, &FGameplayTagQueryCustomization::GetQueryDescText)
-		.AutoWrapText(true)
-		]
-		]
+			+SHorizontalBox::Slot()
+			.AutoWidth()
+			.VAlign(VAlign_Center)
+			[
+				SNew(SVerticalBox)
+				+SVerticalBox::Slot()
+				.AutoHeight()
+				[
+					SNew(SButton)
+					.Text(this, &FGameplayTagQueryCustomization::GetEditButtonText)
+					.OnClicked(this, &FGameplayTagQueryCustomization::OnEditButtonClicked)
+				]
+				+SVerticalBox::Slot()
+				.AutoHeight()
+				[
+					SNew(SButton)
+					.IsEnabled(!bReadOnly)
+					.Text(LOCTEXT("GameplayTagQueryCustomization_Clear", "Clear All"))
+					.OnClicked(this, &FGameplayTagQueryCustomization::OnClearAllButtonClicked)
+					.Visibility(this, &FGameplayTagQueryCustomization::GetClearAllVisibility)
+				]
+			]
+			+SHorizontalBox::Slot()
+			.AutoWidth()
+			[
+				SNew(SBorder)
+				.Padding(4.0f)
+				.Visibility(this, &FGameplayTagQueryCustomization::GetQueryDescVisibility)
+				[
+					SNew(STextBlock)
+					.Text(this, &FGameplayTagQueryCustomization::GetQueryDescText)
+					.ToolTipText(this, &FGameplayTagQueryCustomization::GetQueryDescText)
+					.AutoWrapText(true)
+				]
+			]
 		];
 
 	GEditor->RegisterForUndo(this);
