GameplayTagContainerCustomization registers with OnObjectPropertyChanged delegate to listen for changes in the GameplayTagContainer property it's displaying. If change is detected it refreshes the tag list

Fixes issue where GameplayTagContainer details view didn't update the tag list when the GameplayTagContainer was modified by something other than the details customization widget

diff --git a/Engine/Plugins/Editor/GameplayTagsEditor/Source/GameplayTagsEditor/Private/GameplayTagContainerCustomization.cpp b/Engine/Plugins/Editor/GameplayTagsEditor/Source/GameplayTagsEditor/Private/GameplayTagContainerCustomization.cpp
--- a/Engine/Plugins/Editor/GameplayTagsEditor/Source/GameplayTagsEditor/Private/GameplayTagContainerCustomization.cpp
+++ b/Engine/Plugins/Editor/GameplayTagsEditor/Source/GameplayTagsEditor/Private/GameplayTagContainerCustomization.cpp
@@ -23,6 +23,8 @@ void FGameplayTagContainerCustomization::CustomizeHeader(TSharedRef<class IPrope
 	FSimpleDelegate OnTagContainerChanged = FSimpleDelegate::CreateSP(this, &FGameplayTagContainerCustomization::RefreshTagList);
 	StructPropertyHandle->SetOnPropertyValueChanged(OnTagContainerChanged);
 
+	OnObjectPostEditChangeHandle = FCoreUObjectDelegates::OnObjectPropertyChanged.AddSP(this, &FGameplayTagContainerCustomization::OnObjectPostEditChange);
+
 	BuildEditableContainerList();
 
 	FUIAction SearchForReferencesAction(FExecuteAction::CreateSP(this, &FGameplayTagContainerCustomization::OnWholeContainerSearchForReferences));
@@ -270,6 +272,14 @@ FReply FGameplayTagContainerCustomization::OnRemoveTagClicked(FGameplayTag Tag)
 	return FReply::Handled();
 }
 
+void FGameplayTagContainerCustomization::OnObjectPostEditChange(class UObject* Object, FPropertyChangedEvent& PropertyChangedEvent)
+{
+	if (StructPropertyHandle.IsValid() && StructPropertyHandle->GetProperty() && StructPropertyHandle->GetProperty()->GetFName() == PropertyChangedEvent.GetPropertyName())
+	{
+		RefreshTagList();
+	}
+}
+
 TSharedRef<SWidget> FGameplayTagContainerCustomization::GetListContent()
 {
 	if (!StructPropertyHandle.IsValid() || StructPropertyHandle->GetProperty() == nullptr)
@@ -352,6 +362,9 @@ void FGameplayTagContainerCustomization::PostRedo( bool bSuccess )
 FGameplayTagContainerCustomization::~FGameplayTagContainerCustomization()
 {
 	GEditor->UnregisterForUndo(this);
+	FCoreUObjectDelegates::OnObjectPropertyChanged.Remove(OnObjectPostEditChangeHandle);
+	OnObjectPostEditChangeHandle.Reset();
+
 }
 
 void FGameplayTagContainerCustomization::BuildEditableContainerList()
diff --git a/Engine/Plugins/Editor/GameplayTagsEditor/Source/GameplayTagsEditor/Private/GameplayTagContainerCustomization.h b/Engine/Plugins/Editor/GameplayTagsEditor/Source/GameplayTagsEditor/Private/GameplayTagContainerCustomization.h
--- a/Engine/Plugins/Editor/GameplayTagsEditor/Source/GameplayTagsEditor/Private/GameplayTagContainerCustomization.h
+++ b/Engine/Plugins/Editor/GameplayTagsEditor/Source/GameplayTagsEditor/Private/GameplayTagContainerCustomization.h
@@ -10,6 +10,7 @@
 #include "Widgets/Views/STableRow.h"
 #include "SGameplayTagWidget.h"
 #include "EditorUndoClient.h"
+#include "UObject/UnrealType.h"
 
 class IPropertyHandle;
 
@@ -82,6 +83,9 @@ private:
 	void OnTagDoubleClicked(FGameplayTag Tag);
 	FReply OnRemoveTagClicked(FGameplayTag Tag);
 
+	void OnObjectPostEditChange(class UObject* Object, FPropertyChangedEvent& PropertyChangedEvent);
+	FDelegateHandle OnObjectPostEditChangeHandle;
+
 	FReply OnSingleTagMouseButtonPressed(const FGeometry& MyGeometry, const FPointerEvent& MouseEvent, FString TagName);
 	void OnSingleTagSearchForReferences(FString TagName);
 	void OnWholeContainerSearchForReferences();
