Removed the cache of editable gameplay tag containers when editing a known property handle containing gameplay tag containers

This avoids issues where the cache of editable handles gets stale if the underlying object being edited in the detail panel is re-instanced, or if the handle is otherwise updated via IPropertyHandle::ReplaceOuterObjects (which emits no event to say the change has happened).

diff --git a/Engine/Plugins/Editor/GameplayTagsEditor/Source/GameplayTagsEditor/Private/GameplayTagContainerCustomization.cpp b/Engine/Plugins/Editor/GameplayTagsEditor/Source/GameplayTagsEditor/Private/GameplayTagContainerCustomization.cpp
--- a/Engine/Plugins/Editor/GameplayTagsEditor/Source/GameplayTagsEditor/Private/GameplayTagContainerCustomization.cpp
+++ b/Engine/Plugins/Editor/GameplayTagsEditor/Source/GameplayTagsEditor/Private/GameplayTagContainerCustomization.cpp
@@ -23,8 +23,6 @@ void FGameplayTagContainerCustomization::CustomizeHeader(TSharedRef<class IPrope
 	FSimpleDelegate OnTagContainerChanged = FSimpleDelegate::CreateSP(this, &FGameplayTagContainerCustomization::RefreshTagList);
 	StructPropertyHandle->SetOnPropertyValueChanged(OnTagContainerChanged);
 
-	BuildEditableContainerList();
-
 	FUIAction SearchForReferencesAction(FExecuteAction::CreateSP(this, &FGameplayTagContainerCustomization::OnWholeContainerSearchForReferences));
 
 	HeaderRow
@@ -98,20 +96,21 @@ TSharedRef<SWidget> FGameplayTagContainerCustomization::ActiveTags()
 
 void FGameplayTagContainerCustomization::RefreshTagList()
 {
-	// Rebuild Editable Containers as container references can become unsafe
-	BuildEditableContainerList();
-
 	// Build the set of tags on any instance, collapsing common tags together
 	TSet<FGameplayTag> CurrentTagSet;
-	for (int32 ContainerIdx = 0; ContainerIdx < EditableContainers.Num(); ++ContainerIdx)
+	if (StructPropertyHandle.IsValid())
 	{
-		if (const FGameplayTagContainer* Container = EditableContainers[ContainerIdx].TagContainer)
+		SGameplayTagWidget::EnumerateEditableTagContainersFromPropertyHandle(StructPropertyHandle.ToSharedRef(), [&CurrentTagSet](const SGameplayTagWidget::FEditableGameplayTagContainerDatum& EditableTagContainer)
 		{
-			for (auto It = Container->CreateConstIterator(); It; ++It)
+			if (const FGameplayTagContainer* Container = EditableTagContainer.TagContainer)
 			{
-				CurrentTagSet.Add(*It);
+				for (auto It = Container->CreateConstIterator(); It; ++It)
+				{
+					CurrentTagSet.Add(*It);
+				}
 			}
-		}
+			return true;
+		});
 	}
 
 	// Convert the set into pointers for the combo
@@ -248,17 +247,18 @@ void FGameplayTagContainerCustomization::OnTagDoubleClicked(FGameplayTag Tag)
 FReply FGameplayTagContainerCustomization::OnRemoveTagClicked(FGameplayTag Tag)
 {
 	TArray<FString> NewValues;
-	for (int32 ContainerIdx = 0; ContainerIdx < EditableContainers.Num(); ++ContainerIdx)
+	SGameplayTagWidget::EnumerateEditableTagContainersFromPropertyHandle(StructPropertyHandle.ToSharedRef(), [&Tag, &NewValues](const SGameplayTagWidget::FEditableGameplayTagContainerDatum& EditableTagContainer)
 	{
 		FGameplayTagContainer TagContainerCopy;
-		if (const FGameplayTagContainer* Container = EditableContainers[ContainerIdx].TagContainer)
+		if (const FGameplayTagContainer* Container = EditableTagContainer.TagContainer)
 		{
 			TagContainerCopy = *Container;
 		}
 		TagContainerCopy.RemoveTag(Tag);
 
 		NewValues.Add(TagContainerCopy.ToString());
-	}
+		return true;
+	});
 
 	{
 		FScopedTransaction Transaction(LOCTEXT("RemoveGameplayTagFromContainer", "Remove Gameplay Tag"));
@@ -281,7 +281,7 @@ TSharedRef<SWidget> FGameplayTagContainerCustomization::GetListContent()
 
 	bool bReadOnly = StructPropertyHandle->IsEditConst();
 
-	TSharedRef<SGameplayTagWidget> TagWidget = SNew(SGameplayTagWidget, EditableContainers)
+	TSharedRef<SGameplayTagWidget> TagWidget = SNew(SGameplayTagWidget, TArray<SGameplayTagWidget::FEditableGameplayTagContainerDatum>()) // empty container list as built from StructPropertyHandle
 		.Filter(Categories)
 		.ReadOnly(bReadOnly)
 		.TagContainerName(StructPropertyHandle->GetPropertyDisplayName().ToString())
@@ -354,20 +354,4 @@ FGameplayTagContainerCustomization::~FGameplayTagContainerCustomization()
 	GEditor->UnregisterForUndo(this);
 }
 
-void FGameplayTagContainerCustomization::BuildEditableContainerList()
-{
-	EditableContainers.Empty();
-
-	if( StructPropertyHandle.IsValid() )
-	{
-		TArray<void*> RawStructData;
-		StructPropertyHandle->AccessRawData(RawStructData);
-
-		for (int32 ContainerIdx = 0; ContainerIdx < RawStructData.Num(); ++ContainerIdx)
-		{
-			EditableContainers.Add(SGameplayTagWidget::FEditableGameplayTagContainerDatum(nullptr, (FGameplayTagContainer*)RawStructData[ContainerIdx]));
-		}
-	}	
-}
-
 #undef LOCTEXT_NAMESPACE
diff --git a/Engine/Plugins/Editor/GameplayTagsEditor/Source/GameplayTagsEditor/Private/GameplayTagContainerCustomization.h b/Engine/Plugins/Editor/GameplayTagsEditor/Source/GameplayTagsEditor/Private/GameplayTagContainerCustomization.h
--- a/Engine/Plugins/Editor/GameplayTagsEditor/Source/GameplayTagsEditor/Private/GameplayTagContainerCustomization.h
+++ b/Engine/Plugins/Editor/GameplayTagsEditor/Source/GameplayTagsEditor/Private/GameplayTagContainerCustomization.h
@@ -55,9 +55,6 @@ private:
 	/** On Generate Row Delegate */
 	TSharedRef<ITableRow> MakeListViewWidget(TSharedPtr<FGameplayTag> Item, const TSharedRef<STableViewBase>& OwnerTable);
 
-	/** Build List of Editable Containers */
-	void BuildEditableContainerList();
-
 	/** Callback function to create content for the combo button. */
 	TSharedRef<SWidget> GetListContent();
 
@@ -66,9 +63,6 @@ private:
 	/** Cached property handle */
 	TSharedPtr<IPropertyHandle> StructPropertyHandle;
 
-	/** The array of containers this objects has */
-	TArray<SGameplayTagWidget::FEditableGameplayTagContainerDatum> EditableContainers;
-
 	/** List of tags selected in the tag containers */
 	TArray< TSharedPtr<FGameplayTag> > TagList;
 
diff --git a/Engine/Plugins/Editor/GameplayTagsEditor/Source/GameplayTagsEditor/Private/SGameplayTagWidget.cpp b/Engine/Plugins/Editor/GameplayTagsEditor/Source/GameplayTagsEditor/Private/SGameplayTagWidget.cpp
--- a/Engine/Plugins/Editor/GameplayTagsEditor/Source/GameplayTagsEditor/Private/SGameplayTagWidget.cpp
+++ b/Engine/Plugins/Editor/GameplayTagsEditor/Source/GameplayTagsEditor/Private/SGameplayTagWidget.cpp
@@ -37,11 +37,46 @@
 
 const FString SGameplayTagWidget::SettingsIniSection = TEXT("GameplayTagWidget");
 
+bool SGameplayTagWidget::EnumerateEditableTagContainersFromPropertyHandle(const TSharedRef<IPropertyHandle>& PropHandle, TFunctionRef<bool(const FEditableGameplayTagContainerDatum&)> Callback)
+{
+	FStructProperty* StructProperty = CastField<FStructProperty>(PropHandle->GetProperty());
+	if (StructProperty && StructProperty->Struct->IsChildOf(FGameplayTagContainer::StaticStruct()))
+	{
+		PropHandle->EnumerateRawData([&Callback](void* RawData, const int32 /*DataIndex*/, const int32 /*NumDatas*/)
+		{
+			return Callback(FEditableGameplayTagContainerDatum(nullptr, static_cast<FGameplayTagContainer*>(RawData)));
+		});
+		return true;
+	}
+	return false;
+}
+
+bool SGameplayTagWidget::GetEditableTagContainersFromPropertyHandle(const TSharedRef<IPropertyHandle>& PropHandle, TArray<FEditableGameplayTagContainerDatum>& OutEditableTagContainers)
+{
+	FStructProperty* StructProperty = CastField<FStructProperty>(PropHandle->GetProperty());
+	if (StructProperty && StructProperty->Struct->IsChildOf(FGameplayTagContainer::StaticStruct()))
+	{
+		OutEditableTagContainers.Reset();
+		return EnumerateEditableTagContainersFromPropertyHandle(PropHandle, [&OutEditableTagContainers](const FEditableGameplayTagContainerDatum& EditableTagContainer)
+		{
+			OutEditableTagContainers.Add(EditableTagContainer);
+			return true;
+		});
+	}
+	return false;
+}
+
 void SGameplayTagWidget::Construct(const FArguments& InArgs, const TArray<FEditableGameplayTagContainerDatum>& EditableTagContainers)
 {
-	// If we're in management mode, we don't need to have editable tag containers.
-	ensure(EditableTagContainers.Num() > 0 || InArgs._GameplayTagUIMode == EGameplayTagUIMode::ManagementMode);
 	TagContainers = EditableTagContainers;
+	if (InArgs._PropertyHandle.IsValid())
+	{
+		// If we're backed by a property handle then try and get the tag containers from the property handle
+		GetEditableTagContainersFromPropertyHandle(InArgs._PropertyHandle.ToSharedRef(), TagContainers);
+	}
+
+	// If we're in management mode, we don't need to have editable tag containers.
+	ensure(TagContainers.Num() > 0 || InArgs._GameplayTagUIMode == EGameplayTagUIMode::ManagementMode);
 
 	OnTagChanged = InArgs._OnTagChanged;
 	bReadOnly = InArgs._ReadOnly;
@@ -288,6 +323,12 @@ void SGameplayTagWidget::Tick(const FGeometry& AllottedGeometry, const double In
 {
 	SCompoundWidget::Tick(AllottedGeometry, InCurrentTime, InDeltaTime);
 
+	if (PropertyHandle.IsValid())
+		// If we're backed by a property handle then try and refresh the tag containers, 
+		// as they may have changed under us (eg, from object re-instancing)
+		GetEditableTagContainersFromPropertyHandle(PropertyHandle.ToSharedRef(), TagContainers);
+	}
+
 	if (bDelayRefresh)
 	{
 		RefreshTags();
diff --git a/Engine/Plugins/Editor/GameplayTagsEditor/Source/GameplayTagsEditor/Public/SGameplayTagWidget.h b/Engine/Plugins/Editor/GameplayTagsEditor/Source/GameplayTagsEditor/Public/SGameplayTagWidget.h
--- a/Engine/Plugins/Editor/GameplayTagsEditor/Source/GameplayTagsEditor/Public/SGameplayTagWidget.h
+++ b/Engine/Plugins/Editor/GameplayTagsEditor/Source/GameplayTagsEditor/Public/SGameplayTagWidget.h
@@ -75,6 +75,18 @@ public:
 		struct FGameplayTagContainer* TagContainer; 
 	};
 
+	/**
+	 * Given a property handle, try and enumerate the editable tag containers from within it (when dealing with a struct property of type FGameplayTagContainer).
+	 * @return True if it was possible to enumerate containers (even if no containers were enumerated), or false otherwise.
+	 */
+	static bool EnumerateEditableTagContainersFromPropertyHandle(const TSharedRef<IPropertyHandle>& PropHandle, TFunctionRef<bool(const FEditableGameplayTagContainerDatum&)> Callback);
+
+	/**
+	 * Given a property handle, try and extract the editable tag containers from within it (when dealing with a struct property of type FGameplayTagContainer).
+	 * @return True if it was possible to extract containers (even if no containers were extracted), or false otherwise.
+	 */
+	static bool GetEditableTagContainersFromPropertyHandle(const TSharedRef<IPropertyHandle>& PropHandle, TArray<FEditableGameplayTagContainerDatum>& OutEditableContainers);
+
 	/** Construct the actual widget */
 	void Construct(const FArguments& InArgs, const TArray<FEditableGameplayTagContainerDatum>& EditableTagContainers);
 
