Repair bad auto-merge

diff --git a/Engine/Plugins/Editor/GameplayTagsEditor/Source/GameplayTagsEditor/Private/SGameplayTagWidget.cpp b/Engine/Plugins/Editor/GameplayTagsEditor/Source/GameplayTagsEditor/Private/SGameplayTagWidget.cpp
--- a/Engine/Plugins/Editor/GameplayTagsEditor/Source/GameplayTagsEditor/Private/SGameplayTagWidget.cpp
+++ b/Engine/Plugins/Editor/GameplayTagsEditor/Source/GameplayTagsEditor/Private/SGameplayTagWidget.cpp
@@ -324,6 +324,7 @@ void SGameplayTagWidget::Tick(const FGeometry& AllottedGeometry, const double In
 	SCompoundWidget::Tick(AllottedGeometry, InCurrentTime, InDeltaTime);
 
 	if (PropertyHandle.IsValid())
+	{
 		// If we're backed by a property handle then try and refresh the tag containers, 
 		// as they may have changed under us (eg, from object re-instancing)
 		GetEditableTagContainersFromPropertyHandle(PropertyHandle.ToSharedRef(), TagContainers);
