Exposing GameplayTagContainerCustomizationPublic for consistency with existing GameplayTagCustomizationPublic (allows specialized derived struct types to leverage same customization)

diff --git a/Engine/Plugins/Editor/GameplayTagsEditor/Source/GameplayTagsEditor/Private/GameplayTagContainerCustomization.cpp b/Engine/Plugins/Editor/GameplayTagsEditor/Source/GameplayTagsEditor/Private/GameplayTagContainerCustomization.cpp
--- a/Engine/Plugins/Editor/GameplayTagsEditor/Source/GameplayTagsEditor/Private/GameplayTagContainerCustomization.cpp
+++ b/Engine/Plugins/Editor/GameplayTagsEditor/Source/GameplayTagsEditor/Private/GameplayTagContainerCustomization.cpp
@@ -9,6 +9,7 @@
 #include "Editor.h"
 #include "PropertyHandle.h"
 #include "DetailWidgetRow.h"
+#include "GameplayTagsEditorModule.h"
 #include "ScopedTransaction.h"
 #include "Widgets/Input/SHyperlink.h"
 #include "EditorFontGlyphs.h"
@@ -16,6 +17,11 @@
 
 #define LOCTEXT_NAMESPACE "GameplayTagContainerCustomization"
 
+TSharedRef<IPropertyTypeCustomization> FGameplayTagContainerCustomizationPublic::MakeInstance()
+{
+	return MakeShareable(new FGameplayTagContainerCustomization);
+}
+
 void FGameplayTagContainerCustomization::CustomizeHeader(TSharedRef<class IPropertyHandle> InStructPropertyHandle, class FDetailWidgetRow& HeaderRow, IPropertyTypeCustomizationUtils& StructCustomizationUtils)
 {
 	StructPropertyHandle = InStructPropertyHandle;
diff --git a/Engine/Plugins/Editor/GameplayTagsEditor/Source/GameplayTagsEditor/Public/GameplayTagsEditorModule.h b/Engine/Plugins/Editor/GameplayTagsEditor/Source/GameplayTagsEditor/Public/GameplayTagsEditorModule.h
--- a/Engine/Plugins/Editor/GameplayTagsEditor/Source/GameplayTagsEditor/Public/GameplayTagsEditorModule.h
+++ b/Engine/Plugins/Editor/GameplayTagsEditor/Source/GameplayTagsEditor/Public/GameplayTagsEditorModule.h
@@ -83,7 +83,13 @@ struct GAMEPLAYTAGSEDITOR_API FGameplayTagCustomizationPublic
 	static TSharedRef<IPropertyTypeCustomization> MakeInstance();
 };
 
+/** This is public so that child structs of FGameplayTagContainer can use the details customization */
+struct GAMEPLAYTAGSEDITOR_API FGameplayTagContainerCustomizationPublic
+{
+	static TSharedRef<IPropertyTypeCustomization> MakeInstance();
+};
+
 struct GAMEPLAYTAGSEDITOR_API FRestrictedGameplayTagCustomizationPublic
 {
 	static TSharedRef<IPropertyTypeCustomization> MakeInstance();
-};
\ No newline at end of file
+};
