Creating and exposing public GameplayTagCustomizationOptions  and GameplayTagContainerCustomizationOptions to allow external property customization handling to influence behaviour of core customization.
Added initial settings to allow Add Tag / Add Tag Source widgets to be optionally hidden by end user

diff --git a/Engine/Plugins/Editor/GameplayTagsEditor/Source/GameplayTagsEditor/Private/GameplayTagContainerCustomization.cpp b/Engine/Plugins/Editor/GameplayTagsEditor/Source/GameplayTagsEditor/Private/GameplayTagContainerCustomization.cpp
--- a/Engine/Plugins/Editor/GameplayTagsEditor/Source/GameplayTagsEditor/Private/GameplayTagContainerCustomization.cpp
+++ b/Engine/Plugins/Editor/GameplayTagsEditor/Source/GameplayTagsEditor/Private/GameplayTagContainerCustomization.cpp
@@ -19,9 +19,18 @@
 
 TSharedRef<IPropertyTypeCustomization> FGameplayTagContainerCustomizationPublic::MakeInstance()
 {
-	return MakeShareable(new FGameplayTagContainerCustomization);
+	return MakeInstanceWithOptions({});
 }
 
+TSharedRef<IPropertyTypeCustomization> FGameplayTagContainerCustomizationPublic::MakeInstanceWithOptions(const FGameplayTagContainerCustomizationOptions& Options)
+{
+	return MakeShareable(new FGameplayTagContainerCustomization(Options));
+}
+
+FGameplayTagContainerCustomization::FGameplayTagContainerCustomization(const FGameplayTagContainerCustomizationOptions& InOptions):
+	Options(InOptions)
+{}
+
 void FGameplayTagContainerCustomization::CustomizeHeader(TSharedRef<class IPropertyHandle> InStructPropertyHandle, class FDetailWidgetRow& HeaderRow, IPropertyTypeCustomizationUtils& StructCustomizationUtils)
 {
 	StructPropertyHandle = InStructPropertyHandle;
@@ -292,7 +301,9 @@ TSharedRef<SWidget> FGameplayTagContainerCustomization::GetListContent()
 		.ReadOnly(bReadOnly)
 		.TagContainerName(StructPropertyHandle->GetPropertyDisplayName().ToString())
 		.OnTagChanged(this, &FGameplayTagContainerCustomization::RefreshTagList)
-		.PropertyHandle(StructPropertyHandle);
+		.PropertyHandle(StructPropertyHandle)
+		.ForceHideAddNewTag(Options.bForceHideAddTag)
+		.ForceHideAddNewTagSource(Options.bForceHideAddTagSource);
 
 	LastTagWidget = TagWidget;
 
diff --git a/Engine/Plugins/Editor/GameplayTagsEditor/Source/GameplayTagsEditor/Private/GameplayTagContainerCustomization.h b/Engine/Plugins/Editor/GameplayTagsEditor/Source/GameplayTagsEditor/Private/GameplayTagContainerCustomization.h
--- a/Engine/Plugins/Editor/GameplayTagsEditor/Source/GameplayTagsEditor/Private/GameplayTagContainerCustomization.h
+++ b/Engine/Plugins/Editor/GameplayTagsEditor/Source/GameplayTagsEditor/Private/GameplayTagContainerCustomization.h
@@ -6,6 +6,7 @@
 #include "Layout/Visibility.h"
 #include "Widgets/SWidget.h"
 #include "Editor/PropertyEditor/Public/IPropertyTypeCustomization.h"
+#include "GameplayTagContainerCustomizationOptions.h"
 #include "Widgets/Views/STableViewBase.h"
 #include "Widgets/Views/STableRow.h"
 #include "SGameplayTagWidget.h"
@@ -17,11 +18,8 @@ class IPropertyHandle;
 class FGameplayTagContainerCustomization : public IPropertyTypeCustomization, public FEditorUndoClient
 {
 public:
-	static TSharedRef<IPropertyTypeCustomization> MakeInstance()
-	{
-		return MakeShareable(new FGameplayTagContainerCustomization);
-	}
 
+	FGameplayTagContainerCustomization(const FGameplayTagContainerCustomizationOptions& InOptions);
 	~FGameplayTagContainerCustomization();
 
 	/** Overridden to show an edit button to launch the gameplay tag editor */
@@ -73,6 +71,8 @@ private:
 
 	TWeakPtr<class SGameplayTagWidget> LastTagWidget;
 
+	FGameplayTagContainerCustomizationOptions Options;
+
 	void OnTagDoubleClicked(FGameplayTag Tag);
 	FReply OnRemoveTagClicked(FGameplayTag Tag);
 
diff --git a/Engine/Plugins/Editor/GameplayTagsEditor/Source/GameplayTagsEditor/Private/GameplayTagCustomization.cpp b/Engine/Plugins/Editor/GameplayTagsEditor/Source/GameplayTagsEditor/Private/GameplayTagCustomization.cpp
--- a/Engine/Plugins/Editor/GameplayTagsEditor/Source/GameplayTagsEditor/Private/GameplayTagCustomization.cpp
+++ b/Engine/Plugins/Editor/GameplayTagsEditor/Source/GameplayTagsEditor/Private/GameplayTagCustomization.cpp
@@ -14,9 +14,18 @@
 
 TSharedRef<IPropertyTypeCustomization> FGameplayTagCustomizationPublic::MakeInstance()
 {
-	return MakeShareable(new FGameplayTagCustomization);
+	return MakeInstanceWithOptions({});
 }
 
+TSharedRef<IPropertyTypeCustomization> FGameplayTagCustomizationPublic::MakeInstanceWithOptions(const FGameplayTagCustomizationOptions& Options)
+{
+	return MakeShareable(new FGameplayTagCustomization(Options));
+}
+
+FGameplayTagCustomization::FGameplayTagCustomization(const FGameplayTagCustomizationOptions& InOptions):
+	Options(InOptions)
+{}
+
 void FGameplayTagCustomization::CustomizeHeader(TSharedRef<class IPropertyHandle> InStructPropertyHandle, class FDetailWidgetRow& HeaderRow, IPropertyTypeCustomizationUtils& StructCustomizationUtils)
 {
 	TagContainer = MakeShareable(new FGameplayTagContainer);
@@ -122,7 +131,9 @@ TSharedRef<SWidget> FGameplayTagCustomization::GetListContent()
 		.TagContainerName(StructPropertyHandle->GetPropertyDisplayName().ToString())
 		.MultiSelect(false)
 		.OnTagChanged(this, &FGameplayTagCustomization::OnTagChanged)
-		.PropertyHandle(StructPropertyHandle);
+		.PropertyHandle(StructPropertyHandle)
+		.ForceHideAddNewTag(Options.bForceHideAddTag)
+		.ForceHideAddNewTagSource(Options.bForceHideAddTagSource);
 
 	LastTagWidget = TagWidget;
 
diff --git a/Engine/Plugins/Editor/GameplayTagsEditor/Source/GameplayTagsEditor/Private/GameplayTagCustomization.h b/Engine/Plugins/Editor/GameplayTagsEditor/Source/GameplayTagsEditor/Private/GameplayTagCustomization.h
--- a/Engine/Plugins/Editor/GameplayTagsEditor/Source/GameplayTagsEditor/Private/GameplayTagCustomization.h
+++ b/Engine/Plugins/Editor/GameplayTagsEditor/Source/GameplayTagsEditor/Private/GameplayTagCustomization.h
@@ -5,6 +5,7 @@
 #include "CoreMinimal.h"
 #include "Widgets/SWidget.h"
 #include "Editor/PropertyEditor/Public/IPropertyTypeCustomization.h"
+#include "GameplayTagCustomizationOptions.h"
 #include "SGameplayTagWidget.h"
 #include "EditorUndoClient.h"
 
@@ -15,6 +16,7 @@ class FGameplayTagCustomization : public IPropertyTypeCustomization, public FEdi
 {
 public:
 
+	FGameplayTagCustomization(const FGameplayTagCustomizationOptions& InOptions);
 	~FGameplayTagCustomization();
 
 	/** Overridden to show an edit button to launch the gameplay tag editor */
@@ -67,6 +69,8 @@ private:
 	TSharedPtr<class SComboButton> EditButton;
 
 	TWeakPtr<class SGameplayTagWidget> LastTagWidget;
+
+	FGameplayTagCustomizationOptions Options;
 };
 
 /** Customization for FGameplayTagCreationWidgetHelper showing an add tag button */
diff --git a/Engine/Plugins/Editor/GameplayTagsEditor/Source/GameplayTagsEditor/Private/GameplayTagsEditorModule.cpp b/Engine/Plugins/Editor/GameplayTagsEditor/Source/GameplayTagsEditor/Private/GameplayTagsEditorModule.cpp
--- a/Engine/Plugins/Editor/GameplayTagsEditor/Source/GameplayTagsEditor/Private/GameplayTagsEditorModule.cpp
+++ b/Engine/Plugins/Editor/GameplayTagsEditor/Source/GameplayTagsEditor/Private/GameplayTagsEditorModule.cpp
@@ -56,7 +56,7 @@ public:
 		// Register the details customizer
 		{
 			FPropertyEditorModule& PropertyModule = FModuleManager::LoadModuleChecked<FPropertyEditorModule>("PropertyEditor");
-			PropertyModule.RegisterCustomPropertyTypeLayout("GameplayTagContainer", FOnGetPropertyTypeCustomizationInstance::CreateStatic(&FGameplayTagContainerCustomization::MakeInstance));
+			PropertyModule.RegisterCustomPropertyTypeLayout("GameplayTagContainer", FOnGetPropertyTypeCustomizationInstance::CreateStatic(&FGameplayTagContainerCustomizationPublic::MakeInstance));
 			PropertyModule.RegisterCustomPropertyTypeLayout("GameplayTag", FOnGetPropertyTypeCustomizationInstance::CreateStatic(&FGameplayTagCustomizationPublic::MakeInstance));
 			PropertyModule.RegisterCustomPropertyTypeLayout("GameplayTagQuery", FOnGetPropertyTypeCustomizationInstance::CreateStatic(&FGameplayTagQueryCustomization::MakeInstance));
 			PropertyModule.RegisterCustomPropertyTypeLayout("GameplayTagCreationWidgetHelper", FOnGetPropertyTypeCustomizationInstance::CreateStatic(&FGameplayTagCreationWidgetHelperDetails::MakeInstance));
diff --git a/Engine/Plugins/Editor/GameplayTagsEditor/Source/GameplayTagsEditor/Private/SGameplayTagWidget.cpp b/Engine/Plugins/Editor/GameplayTagsEditor/Source/GameplayTagsEditor/Private/SGameplayTagWidget.cpp
--- a/Engine/Plugins/Editor/GameplayTagsEditor/Source/GameplayTagsEditor/Private/SGameplayTagWidget.cpp
+++ b/Engine/Plugins/Editor/GameplayTagsEditor/Source/GameplayTagsEditor/Private/SGameplayTagWidget.cpp
@@ -85,6 +85,8 @@ void SGameplayTagWidget::Construct(const FArguments& InArgs, const TArray<FEdita
 	PropertyHandle = InArgs._PropertyHandle;
 	RootFilterString = InArgs._Filter;
 	GameplayTagUIMode = InArgs._GameplayTagUIMode;
+	bForceHideAddNewTag = InArgs._ForceHideAddNewTag;
+	bForceHideAddNewTagSource = InArgs._ForceHideAddNewTagSource;
 
 	bAddTagSectionExpanded = InArgs._NewTagControlsInitiallyExpanded;
 	bDelayRefresh = false;
@@ -1259,9 +1261,9 @@ EVisibility SGameplayTagWidget::DetermineExpandableUIVisibility() const
 {
 	UGameplayTagsManager& Manager = UGameplayTagsManager::Get();
 
-	if ( !Manager.ShouldImportTagsFromINI() )
+	if ( !Manager.ShouldImportTagsFromINI() || (bForceHideAddNewTag && bForceHideAddNewTagSource))
 	{
-		// If we can't support adding tags from INI files, we should never see this widget
+		// If we can't support adding tags from INI files, or both options are forcibly disabled, we should never see this widget
 		return EVisibility::Collapsed;
 	}
 
@@ -1270,7 +1272,7 @@ EVisibility SGameplayTagWidget::DetermineExpandableUIVisibility() const
 
 EVisibility SGameplayTagWidget::DetermineAddNewSourceExpandableUIVisibility() const
 {
-	if (bRestrictedTags)
+	if (bRestrictedTags || bForceHideAddNewTagSource)
 	{
 		return EVisibility::Collapsed;
 	}
@@ -1282,7 +1284,7 @@ EVisibility SGameplayTagWidget::DetermineAddNewTagWidgetVisibility() const
 {
 	UGameplayTagsManager& Manager = UGameplayTagsManager::Get();
 
-	if ( !Manager.ShouldImportTagsFromINI() || !bAddTagSectionExpanded || bRestrictedTags )
+	if ( !Manager.ShouldImportTagsFromINI() || !bAddTagSectionExpanded || bRestrictedTags || bForceHideAddNewTag )
 	{
 		// If we can't support adding tags from INI files, we should never see this widget
 		return EVisibility::Collapsed;
@@ -1295,7 +1297,7 @@ EVisibility SGameplayTagWidget::DetermineAddNewRestrictedTagWidgetVisibility() c
 {
 	UGameplayTagsManager& Manager = UGameplayTagsManager::Get();
 
-	if (!Manager.ShouldImportTagsFromINI() || !bAddTagSectionExpanded || !bRestrictedTags)
+	if (!Manager.ShouldImportTagsFromINI() || !bAddTagSectionExpanded || !bRestrictedTags || bForceHideAddNewTag )
 	{
 		// If we can't support adding tags from INI files, we should never see this widget
 		return EVisibility::Collapsed;
@@ -1308,7 +1310,7 @@ EVisibility SGameplayTagWidget::DetermineAddNewSourceWidgetVisibility() const
 {
 	UGameplayTagsManager& Manager = UGameplayTagsManager::Get();
 
-	if (!Manager.ShouldImportTagsFromINI() || !bAddSourceSectionExpanded || bRestrictedTags)
+	if (!Manager.ShouldImportTagsFromINI() || !bAddSourceSectionExpanded || bRestrictedTags || bForceHideAddNewTagSource )
 	{
 		// If we can't support adding tags from INI files, we should never see this widget
 		return EVisibility::Collapsed;
diff --git a/Engine/Plugins/Editor/GameplayTagsEditor/Source/GameplayTagsEditor/Public/GameplayTagContainerCustomizationOptions.h b/Engine/Plugins/Editor/GameplayTagsEditor/Source/GameplayTagsEditor/Public/GameplayTagContainerCustomizationOptions.h
new file mode 100644
--- /dev/null
+++ b/Engine/Plugins/Editor/GameplayTagsEditor/Source/GameplayTagsEditor/Public/GameplayTagContainerCustomizationOptions.h
@@ -0,0 +1,12 @@
+// Copyright Epic Games, Inc. All Rights Reserved.
+
+#pragma once
+
+struct GAMEPLAYTAGSEDITOR_API FGameplayTagContainerCustomizationOptions
+{
+	// If true, any Gameplay Tag Widget created should not offer an 'Add Tag' option 
+	bool bForceHideAddTag = false;
+
+	// If true, any created Gameplay Tag Widget created should not offer an 'Add Tag Source' option 
+	bool bForceHideAddTagSource = false;
+};
\ No newline at end of file
diff --git a/Engine/Plugins/Editor/GameplayTagsEditor/Source/GameplayTagsEditor/Public/GameplayTagCustomizationOptions.h b/Engine/Plugins/Editor/GameplayTagsEditor/Source/GameplayTagsEditor/Public/GameplayTagCustomizationOptions.h
new file mode 100644
--- /dev/null
+++ b/Engine/Plugins/Editor/GameplayTagsEditor/Source/GameplayTagsEditor/Public/GameplayTagCustomizationOptions.h
@@ -0,0 +1,12 @@
+// Copyright Epic Games, Inc. All Rights Reserved.
+
+#pragma once
+
+struct GAMEPLAYTAGSEDITOR_API FGameplayTagCustomizationOptions
+{
+	// If true, any Gameplay Tag Widget created should not offer an 'Add Tag' option 
+	bool bForceHideAddTag = false;
+
+	// If true, any created Gameplay Tag Widget created should not offer an 'Add Tag Source' option 
+	bool bForceHideAddTagSource = false;
+};
\ No newline at end of file
diff --git a/Engine/Plugins/Editor/GameplayTagsEditor/Source/GameplayTagsEditor/Public/GameplayTagsEditorModule.h b/Engine/Plugins/Editor/GameplayTagsEditor/Source/GameplayTagsEditor/Public/GameplayTagsEditorModule.h
--- a/Engine/Plugins/Editor/GameplayTagsEditor/Source/GameplayTagsEditor/Public/GameplayTagsEditorModule.h
+++ b/Engine/Plugins/Editor/GameplayTagsEditor/Source/GameplayTagsEditor/Public/GameplayTagsEditorModule.h
@@ -7,6 +7,8 @@
 #include "Modules/ModuleManager.h"
 #include "IPropertyTypeCustomization.h"
 #include "GameplayTagContainer.h"
+#include "GameplayTagCustomizationOptions.h"
+#include "GameplayTagContainerCustomizationOptions.h"
 
 DECLARE_DELEGATE_OneParam(FOnSetGameplayTag, const FGameplayTag&);
 DECLARE_DELEGATE_OneParam(FOnSetGameplayTagContainer, const FGameplayTagContainer&);
@@ -81,12 +83,14 @@ public:
 struct GAMEPLAYTAGSEDITOR_API FGameplayTagCustomizationPublic
 {
 	static TSharedRef<IPropertyTypeCustomization> MakeInstance();
+	static TSharedRef<IPropertyTypeCustomization> MakeInstanceWithOptions(const FGameplayTagCustomizationOptions& Options);
 };
 
 /** This is public so that child structs of FGameplayTagContainer can use the details customization */
 struct GAMEPLAYTAGSEDITOR_API FGameplayTagContainerCustomizationPublic
 {
 	static TSharedRef<IPropertyTypeCustomization> MakeInstance();
+	static TSharedRef<IPropertyTypeCustomization> MakeInstanceWithOptions(const FGameplayTagContainerCustomizationOptions& Options);
 };
 
 struct GAMEPLAYTAGSEDITOR_API FRestrictedGameplayTagCustomizationPublic
diff --git a/Engine/Plugins/Editor/GameplayTagsEditor/Source/GameplayTagsEditor/Public/SGameplayTagWidget.h b/Engine/Plugins/Editor/GameplayTagsEditor/Source/GameplayTagsEditor/Public/SGameplayTagWidget.h
--- a/Engine/Plugins/Editor/GameplayTagsEditor/Source/GameplayTagsEditor/Public/SGameplayTagWidget.h
+++ b/Engine/Plugins/Editor/GameplayTagsEditor/Source/GameplayTagsEditor/Public/SGameplayTagWidget.h
@@ -45,6 +45,8 @@ public:
 		, _GameplayTagUIMode( EGameplayTagUIMode::SelectionMode )
 		, _MaxHeight(260.0f)
 		, _RestrictedTags( false )
+		, _ForceHideAddNewTag(false)
+		, _ForceHideAddNewTagSource(false)
 	{}
 		SLATE_ARGUMENT( FString, Filter ) // Comma delimited string of tag root names to filter by
 		SLATE_ARGUMENT( FString, NewTagName ) // String that will initially populate the New Tag Name field
@@ -57,6 +59,8 @@ public:
 		SLATE_ARGUMENT( EGameplayTagUIMode, GameplayTagUIMode )	// Determines behavior of the menu based on where it's used
 		SLATE_ARGUMENT( float, MaxHeight )	// caps the height of the gameplay tag tree
 		SLATE_ARGUMENT( bool, RestrictedTags ) // if we are dealing with restricted tags or regular gameplay tags
+		SLATE_ARGUMENT( bool, ForceHideAddNewTag) // Allow caller context to manually manipulate visibility of add new tag widget
+		SLATE_ARGUMENT( bool, ForceHideAddNewTagSource) // Allow caller context to manually manipulate visibility of add new tag source widget
 	SLATE_END_ARGS()
 
 	/** Simple struct holding a tag container and its owner for generic re-use of the widget */
@@ -154,6 +158,12 @@ private:
 	/** If true, this widget is displaying restricted tags; if false this widget displays regular gameplay tags. */
 	bool bRestrictedTags;
 
+	/** If true, this widget is unable to display the 'Add new tag' widget */
+	bool bForceHideAddNewTag;
+
+	/** If true, this widget is unable to display the 'Add new tag source' widget */
+	bool bForceHideAddNewTagSource;
+
 	/** The maximum height of the gameplay tag tree. If 0, the height is unbound. */
 	float MaxHeight;
 
