Forcibly close tag widget when owning customization is destroyed to avoid a lingering window that crashes when you interact

diff --git a/Engine/Plugins/Editor/GameplayTagsEditor/Source/GameplayTagsEditor/Private/GameplayTagContainerCustomization.cpp b/Engine/Plugins/Editor/GameplayTagsEditor/Source/GameplayTagsEditor/Private/GameplayTagContainerCustomization.cpp
--- a/Engine/Plugins/Editor/GameplayTagsEditor/Source/GameplayTagsEditor/Private/GameplayTagContainerCustomization.cpp
+++ b/Engine/Plugins/Editor/GameplayTagsEditor/Source/GameplayTagsEditor/Private/GameplayTagContainerCustomization.cpp
@@ -368,6 +368,12 @@ void FGameplayTagContainerCustomization::PostRedo( bool bSuccess )
 
 FGameplayTagContainerCustomization::~FGameplayTagContainerCustomization()
 {
+	// Forcibly close the popup to avoid crashes later
+	if (EditButton.IsValid() && EditButton->IsOpen())
+	{
+		EditButton->SetIsOpen(false);
+	}
+
 	GEditor->UnregisterForUndo(this);
 }
 
diff --git a/Engine/Plugins/Editor/GameplayTagsEditor/Source/GameplayTagsEditor/Private/GameplayTagCustomization.cpp b/Engine/Plugins/Editor/GameplayTagsEditor/Source/GameplayTagsEditor/Private/GameplayTagCustomization.cpp
--- a/Engine/Plugins/Editor/GameplayTagsEditor/Source/GameplayTagsEditor/Private/GameplayTagCustomization.cpp
+++ b/Engine/Plugins/Editor/GameplayTagsEditor/Source/GameplayTagsEditor/Private/GameplayTagCustomization.cpp
@@ -222,6 +222,12 @@ void FGameplayTagCustomization::PostRedo(bool bSuccess)
 
 FGameplayTagCustomization::~FGameplayTagCustomization()
 {
+	// Forcibly close the popup to avoid crashes later
+	if (EditButton.IsValid() && EditButton->IsOpen())
+	{
+		EditButton->SetIsOpen(false);
+	}
+
 	GEditor->UnregisterForUndo(this);
 }
 
