Add additional options to the GameplayTags widget:
- Add hybrid UI mode to allow selecting tags as well as renaming/deleting them
- Add argument to provide a filter lambda for custom filtering
- Add argument to hide tag tree controls
- Add argument to provide a custom background brush

diff --git a/Engine/Plugins/Editor/GameplayTagsEditor/Source/GameplayTagsEditor/Private/SGameplayTagWidget.cpp b/Engine/Plugins/Editor/GameplayTagsEditor/Source/GameplayTagsEditor/Private/SGameplayTagWidget.cpp
--- a/Engine/Plugins/Editor/GameplayTagsEditor/Source/GameplayTagsEditor/Private/SGameplayTagWidget.cpp
+++ b/Engine/Plugins/Editor/GameplayTagsEditor/Source/GameplayTagsEditor/Private/SGameplayTagWidget.cpp
@@ -32,6 +32,8 @@
 #include "Editor.h"
 #include "Framework/Commands/UIAction.h"
 #include "Framework/MultiBox/MultiBoxBuilder.h"
+#include "Brushes/SlateColorBrush.h"
+#include "Interfaces/IMainFrameModule.h"
 
 #define LOCTEXT_NAMESPACE "GameplayTagWidget"
 
@@ -84,9 +86,11 @@ void SGameplayTagWidget::Construct(const FArguments& InArgs, const TArray<FEdita
 	bMultiSelect = InArgs._MultiSelect;
 	PropertyHandle = InArgs._PropertyHandle;
 	RootFilterString = InArgs._Filter;
+	TagFilter = InArgs._OnFilterTag;
 	GameplayTagUIMode = InArgs._GameplayTagUIMode;
 	bForceHideAddNewTag = InArgs._ForceHideAddNewTag;
 	bForceHideAddNewTagSource = InArgs._ForceHideAddNewTagSource;
+	bForceHideTagTreeControls = InArgs._ForceHideTagTreeControls;
 
 	bAddTagSectionExpanded = InArgs._NewTagControlsInitiallyExpanded;
 	bDelayRefresh = false;
@@ -97,7 +101,7 @@ void SGameplayTagWidget::Construct(const FArguments& InArgs, const TArray<FEdita
 	UGameplayTagsManager::OnEditorRefreshGameplayTagTree.AddSP(this, &SGameplayTagWidget::RefreshOnNextTick);
 	UGameplayTagsManager& Manager = UGameplayTagsManager::Get();
 
-	Manager.GetFilteredGameplayRootTags(RootFilterString, TagItems);
+	GetFilteredGameplayRootTags(RootFilterString, TagItems);
 
 	if (bRestrictedTags)
 	{
@@ -155,7 +159,7 @@ void SGameplayTagWidget::Construct(const FArguments& InArgs, const TArray<FEdita
 	ChildSlot
 	[
 		SNew(SBorder)
-		.BorderImage(FAppStyle::GetBrush("ToolPanel.GroupBorder"))
+		.BorderImage(InArgs._BackgroundBrush)
 		[
 			SNew(SVerticalBox)
 
@@ -253,6 +257,7 @@ void SGameplayTagWidget::Construct(const FArguments& InArgs, const TArray<FEdita
 			.VAlign(VAlign_Top)
 			[
 				SNew(SHorizontalBox)
+				.Visibility(this, &SGameplayTagWidget::DetermineTagTreeControlsVisibility)
 
 				// Expand All nodes
 				+SHorizontalBox::Slot()
@@ -300,6 +305,7 @@ void SGameplayTagWidget::Construct(const FArguments& InArgs, const TArray<FEdita
 			.MaxHeight(MaxHeight)
 			[
 				SAssignNew(TagTreeContainerWidget, SBorder)
+				.BorderImage(FAppStyle::GetBrush("NoBorder"))
 				.Padding(FMargin(4.f))
 				[
 					SAssignNew(TagTreeWidget, STreeView< TSharedPtr<FGameplayTagNode> >)
@@ -313,6 +319,11 @@ void SGameplayTagWidget::Construct(const FArguments& InArgs, const TArray<FEdita
 		]
 	];
 
+	if (InArgs._TagTreeViewBackgroundBrush)
+	{
+		TagTreeWidget->SetBackgroundBrush(InArgs._TagTreeViewBackgroundBrush);
+	}
+
 	// Force the entire tree collapsed to start
 	SetTagTreeItemExpansion(false);
 
@@ -513,7 +524,7 @@ TSharedRef<ITableRow> SGameplayTagWidget::OnGenerateRow(TSharedPtr<FGameplayTagN
 				.IsChecked(this, &SGameplayTagWidget::IsTagChecked, InItem)
 				.ToolTipText(TooltipText)
 				.IsEnabled(this, &SGameplayTagWidget::CanSelectTags)
-				.Visibility( GameplayTagUIMode == EGameplayTagUIMode::SelectionMode ? EVisibility::Visible : EVisibility::Collapsed )
+				.Visibility( GameplayTagUIMode == EGameplayTagUIMode::SelectionMode || GameplayTagUIMode == EGameplayTagUIMode::HybridMode ? EVisibility::Visible : EVisibility::Collapsed )
 				[
 					SNew(STextBlock)
 					.Text(FText::FromName(InItem->GetSimpleTagName()))
@@ -529,7 +540,7 @@ TSharedRef<ITableRow> SGameplayTagWidget::OnGenerateRow(TSharedPtr<FGameplayTagN
 				.ToolTip( FSlateApplication::Get().MakeToolTip(TooltipText) )
 				.Text(FText::FromName( InItem->GetSimpleTagName()) )
 				.ColorAndOpacity(this, &SGameplayTagWidget::GetTagTextColour, InItem)
-				.Visibility( GameplayTagUIMode != EGameplayTagUIMode::SelectionMode ? EVisibility::Visible : EVisibility::Collapsed )
+				.Visibility( GameplayTagUIMode == EGameplayTagUIMode::ManagementMode ? EVisibility::Visible : EVisibility::Collapsed )
 			]
 
 			// Allows non-restricted children checkbox
@@ -578,6 +589,7 @@ TSharedRef<ITableRow> SGameplayTagWidget::OnGenerateRow(TSharedPtr<FGameplayTagN
 				.ForegroundColor(FSlateColor::UseForeground())
 				.HasDownArrow(true)
 				.OnGetMenuContent(this, &SGameplayTagWidget::MakeTagActionsMenu, InItem)
+				.CollapseMenuOnParentFocus(true)
 			]
 		];
 }
@@ -613,8 +625,6 @@ void SGameplayTagWidget::OnTagChecked(TSharedPtr<FGameplayTagNode> NodeChecked)
 {
 	FScopedTransaction Transaction( LOCTEXT("GameplayTagWidget_AddTags", "Add Gameplay Tags") );
 
-	UGameplayTagsManager& TagsManager = UGameplayTagsManager::Get();
-
 	for (int32 ContainerIdx = 0; ContainerIdx < TagContainers.Num(); ++ContainerIdx)
 	{
 		TSharedPtr<FGameplayTagNode> CurNode(NodeChecked);
@@ -915,7 +925,7 @@ FReply SGameplayTagWidget::OnAddSubtagClicked(TSharedPtr<FGameplayTagNode> InTag
 
 TSharedRef<SWidget> SGameplayTagWidget::MakeTagActionsMenu(TSharedPtr<FGameplayTagNode> InTagNode)
 {
-	bool bShowManagement = (GameplayTagUIMode == EGameplayTagUIMode::ManagementMode && !bReadOnly);
+	bool bShowManagement = ((GameplayTagUIMode == EGameplayTagUIMode::ManagementMode || GameplayTagUIMode == EGameplayTagUIMode::HybridMode) && !bReadOnly);
 	UGameplayTagsManager& Manager = UGameplayTagsManager::Get();
 
 	if (!Manager.ShouldImportTagsFromINI())
@@ -956,7 +966,7 @@ TSharedRef<SWidget> SGameplayTagWidget::MakeTagActionsMenu(TSharedPtr<FGameplayT
 	}
 
 	// Only include these menu items if we have tag containers to modify
-	if (TagContainers.Num() > 0)
+	if (TagContainers.Num() > 0 && !bForceHideAddNewTag)
 	{
 		// Either Remove or Add Exact Tag depending on if we have the exact tag or not
 		if (IsExactTagInCollection(InTagNode))
@@ -994,10 +1004,21 @@ void SGameplayTagWidget::OnDeleteTag(TSharedPtr<FGameplayTagNode> InTagNode)
 	if (InTagNode.IsValid())
 	{
 		IGameplayTagsEditorModule& TagsEditor = IGameplayTagsEditorModule::Get();
+		
+		bool TagRemoved = false;
+		
+		if (GameplayTagUIMode == EGameplayTagUIMode::HybridMode)
+		{
+			for (int32 ContainerIdx = 0; ContainerIdx < TagContainers.Num(); ++ContainerIdx)
+			{
+				FGameplayTagContainer* Container = TagContainers[ContainerIdx].TagContainer;
+				TagRemoved |= Container->RemoveTag(InTagNode->GetCompleteTag());
+			}
+		}
 
 		const bool bDeleted = TagsEditor.DeleteTagFromINI(InTagNode);
 
-		if (bDeleted)
+		if (bDeleted || TagRemoved)
 		{
 			OnTagChanged.ExecuteIfBound();
 		}
@@ -1045,7 +1066,7 @@ void SGameplayTagWidget::OnSearchForReferences(TSharedPtr<FGameplayTagNode> InTa
 void SGameplayTagWidget::SetTagTreeItemExpansion(bool bExpand)
 {
 	TArray< TSharedPtr<FGameplayTagNode> > TagArray;
-	UGameplayTagsManager::Get().GetFilteredGameplayRootTags(TEXT(""), TagArray);
+	GetFilteredGameplayRootTags(TEXT(""), TagArray);
 	for (int32 TagIdx = 0; TagIdx < TagArray.Num(); ++TagIdx)
 	{
 		SetTagNodeItemExpansion(TagArray[TagIdx], bExpand);
@@ -1072,7 +1093,7 @@ void SGameplayTagWidget::LoadSettings()
 	MigrateSettings();
 
 	TArray< TSharedPtr<FGameplayTagNode> > TagArray;
-	UGameplayTagsManager::Get().GetFilteredGameplayRootTags(TEXT(""), TagArray);
+	GetFilteredGameplayRootTags(TEXT(""), TagArray);
 	for (int32 TagIdx = 0; TagIdx < TagArray.Num(); ++TagIdx)
 	{
 		LoadTagNodeItemExpansion(TagArray[TagIdx] );
@@ -1179,7 +1200,7 @@ void SGameplayTagWidget::SetContainer(FGameplayTagContainer* OriginalContainer,
 	}
 	else
 	{
-		// Not sure if we should get here, means the property handle hasnt been setup which could be right or wrong.
+		// Not sure if we should get here, means the property handle hasn't been setup which could be right or wrong.
 		if (OwnerObj)
 		{
 			OwnerObj->PreEditChange(PropertyHandle.IsValid() ? PropertyHandle->GetProperty() : nullptr);
@@ -1225,7 +1246,7 @@ void SGameplayTagWidget::OnGameplayTagAdded(const FString& TagName, const FStrin
 void SGameplayTagWidget::RefreshTags()
 {
 	UGameplayTagsManager& Manager = UGameplayTagsManager::Get();
-	Manager.GetFilteredGameplayRootTags(RootFilterString, TagItems);
+	GetFilteredGameplayRootTags(RootFilterString, TagItems);
 
 	if (bRestrictedTags)
 	{
@@ -1318,6 +1339,16 @@ EVisibility SGameplayTagWidget::DetermineAddNewSourceWidgetVisibility() const
 	return EVisibility::Visible;
 }
 
+EVisibility SGameplayTagWidget::DetermineTagTreeControlsVisibility() const
+{
+	if (bForceHideTagTreeControls)
+	{
+		return EVisibility::Collapsed;
+	}
+
+	return EVisibility::Visible;
+}
+
 EVisibility SGameplayTagWidget::DetermineAddNewSubTagWidgetVisibility(TSharedPtr<FGameplayTagNode> Node) const
 {
 	EVisibility LocalVisibility = DetermineExpandableUIVisibility();
@@ -1348,7 +1379,7 @@ EVisibility SGameplayTagWidget::DetermineClearSelectionVisibility() const
 
 bool SGameplayTagWidget::CanSelectTags() const
 {
-	return !bReadOnly && GameplayTagUIMode == EGameplayTagUIMode::SelectionMode;
+	return !bReadOnly && (GameplayTagUIMode == EGameplayTagUIMode::SelectionMode || GameplayTagUIMode == EGameplayTagUIMode::HybridMode);
 }
 
 bool SGameplayTagWidget::IsAddingNewTag() const
@@ -1403,9 +1434,18 @@ void SGameplayTagWidget::OpenRenameGameplayTagDialog(TSharedPtr<FGameplayTagNode
 
 	RenameTagWindow->SetContent(RenameTagDialog);
 
-	TSharedPtr<SWindow> CurrentWindow = FSlateApplication::Get().FindWidgetWindow(AsShared() );
-
-	FSlateApplication::Get().AddModalWindow(RenameTagWindow, CurrentWindow);
+	const TSharedPtr<SWindow> CurrentWindow = FSlateApplication::Get().FindWidgetWindow(AsShared());
+	
+	if (CurrentWindow->IsRegularWindow())
+	{
+		FSlateApplication::Get().AddModalWindow(RenameTagWindow, CurrentWindow);
+	}
+	else
+	{
+		// AddModalWindow will crash if the parent window is a pop up, so we parent it to the MainFrame window instead.
+		const IMainFrameModule& MainFrame = FModuleManager::LoadModuleChecked<IMainFrameModule>("MainFrame");
+		FSlateApplication::Get().AddModalWindow(RenameTagWindow, MainFrame.GetParentWindow());
+	}
 }
 
 TSharedPtr<SWidget> SGameplayTagWidget::GetWidgetToFocusOnOpen()
@@ -1461,4 +1501,27 @@ void SGameplayTagWidget::VerifyAssetTagValidity()
 	}
 }
 
+void SGameplayTagWidget::GetFilteredGameplayRootTags(const FString& InFilterString, TArray<TSharedPtr<FGameplayTagNode>>& OutNodes) const
+{
+	OutNodes.Empty();
+	const UGameplayTagsManager& Manager = UGameplayTagsManager::Get();
+
+	if (TagFilter.IsBound())
+	{
+		TArray<TSharedPtr<FGameplayTagNode>> UnfilteredItems;
+		Manager.GetFilteredGameplayRootTags(InFilterString, UnfilteredItems);
+		for (const TSharedPtr<FGameplayTagNode>& Node : UnfilteredItems)
+		{
+			if (TagFilter.Execute(Node) == ETagFilterResult::IncludeTag)
+			{
+				OutNodes.Add(Node);
+			}
+		}
+	}
+	else
+	{
+		Manager.GetFilteredGameplayRootTags(InFilterString, OutNodes);
+	}
+}
+
 #undef LOCTEXT_NAMESPACE
diff --git a/Engine/Plugins/Editor/GameplayTagsEditor/Source/GameplayTagsEditor/Public/SGameplayTagWidget.h b/Engine/Plugins/Editor/GameplayTagsEditor/Source/GameplayTagsEditor/Public/SGameplayTagWidget.h
--- a/Engine/Plugins/Editor/GameplayTagsEditor/Source/GameplayTagsEditor/Public/SGameplayTagWidget.h
+++ b/Engine/Plugins/Editor/GameplayTagsEditor/Source/GameplayTagsEditor/Public/SGameplayTagWidget.h
@@ -8,6 +8,7 @@
 #include "Layout/Visibility.h"
 #include "Widgets/DeclarativeSyntaxSupport.h"
 #include "Input/Reply.h"
+#include "Styling/AppStyle.h"
 #include "Widgets/SWidget.h"
 #include "Widgets/SCompoundWidget.h"
 #include "Widgets/Views/STableViewBase.h"
@@ -22,8 +23,9 @@ class SAddNewGameplayTagSourceWidget;
 /** Determines the behavior of the gameplay tag UI depending on where it's used */
 enum class EGameplayTagUIMode : uint8
 {
-	SelectionMode,
-	ManagementMode,
+	SelectionMode,    // Allows selecting tags for gameplay tag containers.
+	ManagementMode,   // Allows renaming/deleting tags, does not require gameplay tag containers.
+	HybridMode        // Allows doing operations from both Selection and Management mode.
 };
 
 /** Widget allowing user to tag assets with gameplay tags */
@@ -33,6 +35,13 @@ public:
 
 	/** Called when a tag status is changed */
 	DECLARE_DELEGATE( FOnTagChanged )
+	enum class ETagFilterResult
+	{
+		IncludeTag,
+		ExcludeTag
+	};
+	
+	DECLARE_DELEGATE_RetVal_OneParam( ETagFilterResult, FOnFilterTag, const TSharedPtr<FGameplayTagNode>&)
 
 	SLATE_BEGIN_ARGS( SGameplayTagWidget )
 		: _Filter()
@@ -47,8 +56,13 @@ public:
 		, _RestrictedTags( false )
 		, _ForceHideAddNewTag(false)
 		, _ForceHideAddNewTagSource(false)
+		, _ForceHideTagTreeControls(false)
+		, _BackgroundBrush(FAppStyle::GetBrush("ToolPanel.GroupBorder"))
+		, _TagTreeViewBackgroundBrush(nullptr)
+		
 	{}
 		SLATE_ARGUMENT( FString, Filter ) // Comma delimited string of tag root names to filter by
+		SLATE_EVENT( FOnFilterTag, OnFilterTag ) // Optional filter function called when generating the tag list
 		SLATE_ARGUMENT( FString, NewTagName ) // String that will initially populate the New Tag Name field
 		SLATE_ARGUMENT( bool, ReadOnly ) // Flag to set if the list is read only
 		SLATE_ARGUMENT( FString, TagContainerName ) // The name that will be used for the settings file
@@ -59,8 +73,11 @@ public:
 		SLATE_ARGUMENT( EGameplayTagUIMode, GameplayTagUIMode )	// Determines behavior of the menu based on where it's used
 		SLATE_ARGUMENT( float, MaxHeight )	// caps the height of the gameplay tag tree
 		SLATE_ARGUMENT( bool, RestrictedTags ) // if we are dealing with restricted tags or regular gameplay tags
-		SLATE_ARGUMENT( bool, ForceHideAddNewTag) // Allow caller context to manually manipulate visibility of add new tag widget
-		SLATE_ARGUMENT( bool, ForceHideAddNewTagSource) // Allow caller context to manually manipulate visibility of add new tag source widget
+		SLATE_ARGUMENT( bool, ForceHideAddNewTag ) // Allow caller context to manually manipulate visibility of add new tag widget
+		SLATE_ARGUMENT( bool, ForceHideAddNewTagSource ) // Allow caller context to manually manipulate visibility of add new tag source widget
+		SLATE_ARGUMENT( bool, ForceHideTagTreeControls ) // Allow caller context to manually manipulate visibility of Collapse/Expand buttons and filter widget
+		SLATE_ARGUMENT( const FSlateBrush*, BackgroundBrush) // Background brush for the whole widget
+		SLATE_ARGUMENT( const FSlateBrush*, TagTreeViewBackgroundBrush) // Background brush for the whole widget
 	SLATE_END_ARGS()
 
 	/** Simple struct holding a tag container and its owner for generic re-use of the widget */
@@ -140,6 +157,9 @@ private:
 	/** root filter (passed in on creation) */
 	FString RootFilterString;
 
+	/** User specified filter function. */
+	FOnFilterTag TagFilter; 
+
 	/* Flag to set if the list is read only*/
 	bool bReadOnly;
 
@@ -164,6 +184,9 @@ private:
 	/** If true, this widget is unable to display the 'Add new tag source' widget */
 	bool bForceHideAddNewTagSource;
 
+	/** If true, this widget is unable to display the tag tree controls. */
+	bool bForceHideTagTreeControls;
+
 	/** The maximum height of the gameplay tag tree. If 0, the height is unbound. */
 	float MaxHeight;
 
@@ -338,6 +361,9 @@ private:
 	/** Helper function to determine the visibility of the Add New Tag Source widget */
 	EVisibility DetermineAddNewSourceWidgetVisibility() const;
 
+	/** Helper function to determine the visibility of the tag tree controls. */
+	EVisibility DetermineTagTreeControlsVisibility() const;
+
 	/** Helper function to determine the visibility of the Add New Subtag widget */
 	EVisibility DetermineAddNewSubTagWidgetVisibility(TSharedPtr<FGameplayTagNode> Node) const;
 
@@ -399,4 +425,7 @@ private:
 
 	/** Delegate that is fired when a tag is successfully renamed */
 	void OnGameplayTagRenamed(FString OldTagName, FString NewTagName);
+	
+	/** Populate tag items from the gameplay tags manager. */
+	void GetFilteredGameplayRootTags(const FString& InFilterString, TArray<TSharedPtr<FGameplayTagNode>>& OutNodes) const;
 };
