Adding functionality to copy and paste a single tag out of a gameplay tag container

diff --git a/Engine/Plugins/Editor/GameplayTagsEditor/Source/GameplayTagsEditor/GameplayTagsEditor.Build.cs b/Engine/Plugins/Editor/GameplayTagsEditor/Source/GameplayTagsEditor/GameplayTagsEditor.Build.cs
--- a/Engine/Plugins/Editor/GameplayTagsEditor/Source/GameplayTagsEditor/GameplayTagsEditor.Build.cs
+++ b/Engine/Plugins/Editor/GameplayTagsEditor/Source/GameplayTagsEditor/GameplayTagsEditor.Build.cs
@@ -18,6 +18,7 @@ namespace UnrealBuildTool.Rules
 				new string[]
 				{
 					// ... add private dependencies that you statically link with here ...
+					"ApplicationCore",
 					"Core",
 					"CoreUObject",
 					"Engine",
diff --git a/Engine/Plugins/Editor/GameplayTagsEditor/Source/GameplayTagsEditor/Private/GameplayTagContainerCustomization.cpp b/Engine/Plugins/Editor/GameplayTagsEditor/Source/GameplayTagsEditor/Private/GameplayTagContainerCustomization.cpp
--- a/Engine/Plugins/Editor/GameplayTagsEditor/Source/GameplayTagsEditor/Private/GameplayTagContainerCustomization.cpp
+++ b/Engine/Plugins/Editor/GameplayTagsEditor/Source/GameplayTagsEditor/Private/GameplayTagContainerCustomization.cpp
@@ -10,6 +10,7 @@
 #include "PropertyHandle.h"
 #include "DetailWidgetRow.h"
 #include "GameplayTagsEditorModule.h"
+#include "HAL/PlatformApplicationMisc.h"
 #include "ScopedTransaction.h"
 #include "Widgets/Input/SHyperlink.h"
 #include "EditorFontGlyphs.h"
@@ -213,6 +214,30 @@ FReply FGameplayTagContainerCustomization::OnSingleTagMouseButtonPressed(const F
 			FText::Format(LOCTEXT("SingleTagSearchForReferencesTooltip", "Find references to the tag {0}"), FText::AsCultureInvariant(TagName)),
 			FSlateIcon(),
 			SearchForReferencesAction);
+		
+		if (StructPropertyHandle)
+		{
+			FUIAction CopyAction = FUIAction(FExecuteAction::CreateSP(this, &FGameplayTagContainerCustomization::OnCopyTag, TagName));
+			FUIAction PasteAction = FUIAction(FExecuteAction::CreateSP(this, &FGameplayTagContainerCustomization::OnPasteTag),FCanExecuteAction::CreateSP(this, &FGameplayTagContainerCustomization::CanPaste));
+			
+			if (CopyAction.IsBound() && PasteAction.IsBound())
+			{
+				FMenuEntryParams CopyContentParams;
+				CopyContentParams.LabelOverride = NSLOCTEXT("PropertyView", "CopyProperty", "Copy");
+				CopyContentParams.ToolTipOverride = FText::Format(LOCTEXT("SingleTagSearchForReferencesTooltip", "Copy tag {0}"), FText::AsCultureInvariant(TagName));
+				CopyContentParams.IconOverride = FSlateIcon(FCoreStyle::Get().GetStyleSetName(), "GenericCommands.Copy");
+				CopyContentParams.DirectActions = CopyAction;
+				MenuBuilder.AddMenuEntry(CopyContentParams);
+
+				FMenuEntryParams PasteContentParams;
+				PasteContentParams.LabelOverride = NSLOCTEXT("PropertyView", "PasteProperty", "Paste");
+				PasteContentParams.ToolTipOverride = FText::Format(LOCTEXT("SingleTagSearchForReferencesTooltip", "Paste tag {0}"), FText::AsCultureInvariant(TagName));
+				PasteContentParams.IconOverride = FSlateIcon(FCoreStyle::Get().GetStyleSetName(), "GenericCommands.Paste");
+				PasteContentParams.DirectActions = PasteAction;
+				MenuBuilder.AddMenuEntry(PasteContentParams);
+			}
+		}
+		
 		MenuBuilder.EndSection();
 
 		// Spawn context menu
@@ -377,4 +402,50 @@ FGameplayTagContainerCustomization::~FGameplayTagContainerCustomization()
 	GEditor->UnregisterForUndo(this);
 }
 
+void FGameplayTagContainerCustomization::OnCopyTag(FString TagName)
+{
+	FPlatformApplicationMisc::ClipboardCopy(*TagName);
+}
+
+void FGameplayTagContainerCustomization::OnPasteTag()
+{
+	FString TagName;
+	FPlatformApplicationMisc::ClipboardPaste(TagName);
+	
+	const FGameplayTag Tag = FGameplayTag::RequestGameplayTag(FName(TagName), false);
+
+	if (Tag.IsValid())
+	{
+		TArray<FString> NewValues;
+		SGameplayTagWidget::EnumerateEditableTagContainersFromPropertyHandle(StructPropertyHandle.ToSharedRef(), [&Tag, &NewValues](const SGameplayTagWidget::FEditableGameplayTagContainerDatum& EditableTagContainer)
+		{
+			FGameplayTagContainer TagContainerCopy;
+			if (const FGameplayTagContainer* Container = EditableTagContainer.TagContainer)
+			{
+				TagContainerCopy = *Container;
+			}
+			TagContainerCopy.AddTag(Tag);
+
+			NewValues.Add(TagContainerCopy.ToString());
+			return true;
+		});
+
+		{
+			FScopedTransaction Transaction(LOCTEXT("PasteGameplayTagFromContainer", "Paste Gameplay Tag"));
+			StructPropertyHandle->SetPerObjectValues(NewValues);
+		}
+
+		RefreshTagList();
+	}
+}
+
+bool FGameplayTagContainerCustomization::CanPaste()
+{
+	FString TagName;
+	FPlatformApplicationMisc::ClipboardPaste(TagName);
+	const FGameplayTag Tag = FGameplayTag::RequestGameplayTag(FName(TagName), false);
+
+	return Tag.IsValid();
+}
+
 #undef LOCTEXT_NAMESPACE
diff --git a/Engine/Plugins/Editor/GameplayTagsEditor/Source/GameplayTagsEditor/Private/GameplayTagContainerCustomization.h b/Engine/Plugins/Editor/GameplayTagsEditor/Source/GameplayTagsEditor/Private/GameplayTagContainerCustomization.h
--- a/Engine/Plugins/Editor/GameplayTagsEditor/Source/GameplayTagsEditor/Private/GameplayTagContainerCustomization.h
+++ b/Engine/Plugins/Editor/GameplayTagsEditor/Source/GameplayTagsEditor/Private/GameplayTagContainerCustomization.h
@@ -79,5 +79,9 @@ private:
 	FReply OnSingleTagMouseButtonPressed(const FGeometry& MyGeometry, const FPointerEvent& MouseEvent, FString TagName);
 	void OnSingleTagSearchForReferences(FString TagName);
 	void OnWholeContainerSearchForReferences();
+
+	void OnCopyTag(FString TagName);
+	void OnPasteTag();
+	bool CanPaste();
 };
 
