Fix for //UE5/Release-5.1 - Build Engine Localization - GameplayTagContainerCustomization.cpp Warning: Text conflict from LOCTEXT macro for namespace

diff --git a/Engine/Plugins/Editor/GameplayTagsEditor/Source/GameplayTagsEditor/Private/GameplayTagContainerCustomization.cpp b/Engine/Plugins/Editor/GameplayTagsEditor/Source/GameplayTagsEditor/Private/GameplayTagContainerCustomization.cpp
--- a/Engine/Plugins/Editor/GameplayTagsEditor/Source/GameplayTagsEditor/Private/GameplayTagContainerCustomization.cpp
+++ b/Engine/Plugins/Editor/GameplayTagsEditor/Source/GameplayTagsEditor/Private/GameplayTagContainerCustomization.cpp
@@ -224,14 +224,14 @@ FReply FGameplayTagContainerCustomization::OnSingleTagMouseButtonPressed(const F
 			{
 				FMenuEntryParams CopyContentParams;
 				CopyContentParams.LabelOverride = NSLOCTEXT("PropertyView", "CopyProperty", "Copy");
-				CopyContentParams.ToolTipOverride = FText::Format(LOCTEXT("SingleTagSearchForReferencesTooltip", "Copy tag {0}"), FText::AsCultureInvariant(TagName));
+				CopyContentParams.ToolTipOverride = FText::Format(LOCTEXT("SingleTagCopyTooltip", "Copy tag {0}"), FText::AsCultureInvariant(TagName));
 				CopyContentParams.IconOverride = FSlateIcon(FCoreStyle::Get().GetStyleSetName(), "GenericCommands.Copy");
 				CopyContentParams.DirectActions = CopyAction;
 				MenuBuilder.AddMenuEntry(CopyContentParams);
 
 				FMenuEntryParams PasteContentParams;
 				PasteContentParams.LabelOverride = NSLOCTEXT("PropertyView", "PasteProperty", "Paste");
-				PasteContentParams.ToolTipOverride = FText::Format(LOCTEXT("SingleTagSearchForReferencesTooltip", "Paste tag {0}"), FText::AsCultureInvariant(TagName));
+				PasteContentParams.ToolTipOverride = FText::Format(LOCTEXT("SingleTagPasteTooltip", "Paste tag {0}"), FText::AsCultureInvariant(TagName));
 				PasteContentParams.IconOverride = FSlateIcon(FCoreStyle::Get().GetStyleSetName(), "GenericCommands.Paste");
 				PasteContentParams.DirectActions = PasteAction;
 				MenuBuilder.AddMenuEntry(PasteContentParams);
