Fix for a large amount of data in clipboard + trying to search gameplay tag references crashing the editor

diff --git a/Engine/Plugins/Editor/GameplayTagsEditor/Source/GameplayTagsEditor/Private/GameplayTagContainerCustomization.cpp b/Engine/Plugins/Editor/GameplayTagsEditor/Source/GameplayTagsEditor/Private/GameplayTagContainerCustomization.cpp
--- a/Engine/Plugins/Editor/GameplayTagsEditor/Source/GameplayTagsEditor/Private/GameplayTagContainerCustomization.cpp
+++ b/Engine/Plugins/Editor/GameplayTagsEditor/Source/GameplayTagsEditor/Private/GameplayTagContainerCustomization.cpp
@@ -446,9 +446,14 @@ bool FGameplayTagContainerCustomization::CanPaste()
 {
 	FString TagName;
 	FPlatformApplicationMisc::ClipboardPaste(TagName);
-	const FGameplayTag Tag = FGameplayTag::RequestGameplayTag(FName(TagName), false);
-
-	return Tag.IsValid();
+	
+	if (FName::IsValidXName(TagName))
+	{
+		const FGameplayTag Tag = FGameplayTag::RequestGameplayTag(FName(TagName), false);
+		return Tag.IsValid();
+	}
+	
+    return false;
 }
 
 #undef LOCTEXT_NAMESPACE
