Fix for a large amount of data in clipboard + trying to search gameplay tag references crashing the editor (fixed)

diff --git a/Engine/Plugins/Editor/GameplayTagsEditor/Source/GameplayTagsEditor/Private/GameplayTagContainerCustomization.cpp b/Engine/Plugins/Editor/GameplayTagsEditor/Source/GameplayTagsEditor/Private/GameplayTagContainerCustomization.cpp
--- a/Engine/Plugins/Editor/GameplayTagsEditor/Source/GameplayTagsEditor/Private/GameplayTagContainerCustomization.cpp
+++ b/Engine/Plugins/Editor/GameplayTagsEditor/Source/GameplayTagsEditor/Private/GameplayTagContainerCustomization.cpp
@@ -446,8 +446,8 @@ bool FGameplayTagContainerCustomization::CanPaste()
 {
 	FString TagName;
 	FPlatformApplicationMisc::ClipboardPaste(TagName);
-	
-	if (FName::IsValidXName(TagName))
+
+	if (TagName.Len() < NAME_SIZE)
 	{
 		const FGameplayTag Tag = FGameplayTag::RequestGameplayTag(FName(TagName), false);
 		return Tag.IsValid();
