[Engine/Plugins/Editor]

diff --git a/Engine/Plugins/Editor/ContentBrowser/ContentBrowserClassDataSource/Source/ContentBrowserClassDataSource/Private/NewClassContextMenu.h b/Engine/Plugins/Editor/ContentBrowser/ContentBrowserClassDataSource/Source/ContentBrowserClassDataSource/Private/NewClassContextMenu.h
--- a/Engine/Plugins/Editor/ContentBrowser/ContentBrowserClassDataSource/Source/ContentBrowserClassDataSource/Private/NewClassContextMenu.h
+++ b/Engine/Plugins/Editor/ContentBrowser/ContentBrowserClassDataSource/Source/ContentBrowserClassDataSource/Private/NewClassContextMenu.h
@@ -2,7 +2,7 @@
 
 #pragma once
 
-#include "CoreMinimal.h"
+#include "Delegates/Delegate.h"
 
 class UToolMenu;
 
diff --git a/Engine/Plugins/Editor/CryptoKeys/Source/CryptoKeys/Private/CryptoKeysCommandlet.cpp b/Engine/Plugins/Editor/CryptoKeys/Source/CryptoKeys/Private/CryptoKeysCommandlet.cpp
--- a/Engine/Plugins/Editor/CryptoKeys/Source/CryptoKeys/Private/CryptoKeysCommandlet.cpp
+++ b/Engine/Plugins/Editor/CryptoKeys/Source/CryptoKeys/Private/CryptoKeysCommandlet.cpp
@@ -3,7 +3,6 @@
 #include "CryptoKeysCommandlet.h"
 #include "CryptoKeysSettings.h"
 #include "CryptoKeysHelpers.h"
-#include "CryptoKeysOpenSSL.h"
 
 #include UE_INLINE_GENERATED_CPP_BY_NAME(CryptoKeysCommandlet)
 
diff --git a/Engine/Plugins/Editor/CryptoKeys/Source/CryptoKeys/Private/CryptoKeysHelpers.cpp b/Engine/Plugins/Editor/CryptoKeys/Source/CryptoKeys/Private/CryptoKeysHelpers.cpp
--- a/Engine/Plugins/Editor/CryptoKeys/Source/CryptoKeys/Private/CryptoKeysHelpers.cpp
+++ b/Engine/Plugins/Editor/CryptoKeys/Source/CryptoKeys/Private/CryptoKeysHelpers.cpp
@@ -3,10 +3,9 @@
 #include "CryptoKeysHelpers.h"
 #include "CryptoKeysOpenSSL.h"
 
+#include "Logging/LogMacros.h"
 #include "Misc/Base64.h"
-#include "Math/BigInt.h"
 
-#include "CryptoKeysSettings.h"
 
 DEFINE_LOG_CATEGORY_STATIC(LogCryptoKeys, Log, All);
 
diff --git a/Engine/Plugins/Editor/CryptoKeys/Source/CryptoKeys/Private/CryptoKeysSettingsDetails.cpp b/Engine/Plugins/Editor/CryptoKeys/Source/CryptoKeys/Private/CryptoKeysSettingsDetails.cpp
--- a/Engine/Plugins/Editor/CryptoKeys/Source/CryptoKeys/Private/CryptoKeysSettingsDetails.cpp
+++ b/Engine/Plugins/Editor/CryptoKeys/Source/CryptoKeys/Private/CryptoKeysSettingsDetails.cpp
@@ -4,9 +4,9 @@
 #include "CryptoKeysSettings.h"
 #include "CryptoKeysHelpers.h"
 #include "DetailLayoutBuilder.h"
-#include "DetailCategoryBuilder.h"
 #include "DetailWidgetRow.h"
 #include "Widgets/Input/SButton.h"
+#include "Widgets/SBoxPanel.h"
 
 #define LOCTEXT_NAMESPACE "CryptoKeysSettingsDetails"
 
diff --git a/Engine/Plugins/Editor/CryptoKeys/Source/CryptoKeysOpenSSL/Private/CryptoKeysOpenSSL.cpp b/Engine/Plugins/Editor/CryptoKeys/Source/CryptoKeysOpenSSL/Private/CryptoKeysOpenSSL.cpp
--- a/Engine/Plugins/Editor/CryptoKeys/Source/CryptoKeysOpenSSL/Private/CryptoKeysOpenSSL.cpp
+++ b/Engine/Plugins/Editor/CryptoKeys/Source/CryptoKeysOpenSSL/Private/CryptoKeysOpenSSL.cpp
@@ -6,9 +6,10 @@
 #include "Windows/AllowWindowsPlatformTypes.h"
 #endif
 
+#include <openssl/opensslv.h>
 #include <openssl/rsa.h>
+#include <openssl/ossl_typ.h>
 #include <openssl/rand.h>
-#include <openssl/pem.h>
 #include <openssl/bn.h>
 
 #if PLATFORM_WINDOWS
diff --git a/Engine/Plugins/Editor/CurveEditorTools/Source/CurveEditorTools/Private/CurveEditorFocusExtension.cpp b/Engine/Plugins/Editor/CurveEditorTools/Source/CurveEditorTools/Private/CurveEditorFocusExtension.cpp
--- a/Engine/Plugins/Editor/CurveEditorTools/Source/CurveEditorTools/Private/CurveEditorFocusExtension.cpp
+++ b/Engine/Plugins/Editor/CurveEditorTools/Source/CurveEditorTools/Private/CurveEditorFocusExtension.cpp
@@ -1,13 +1,10 @@
 // Copyright Epic Games, Inc. All Rights Reserved.
 
 #include "CurveEditorFocusExtension.h"
-#include "Framework/MultiBox/MultiBoxBuilder.h"
 #include "Framework/Commands/UICommandList.h"
 #include "CurveEditorToolCommands.h"
-#include "ITimeSlider.h"
 #include "CurveEditor.h"
 #include "ICurveEditorBounds.h"
-#include "Framework/MultiBox/MultiBoxExtender.h"
 
 void FCurveEditorFocusExtension::BindCommands(TSharedRef<FUICommandList> CommandBindings)
 {
diff --git a/Engine/Plugins/Editor/CurveEditorTools/Source/CurveEditorTools/Private/CurveEditorToolCommands.h b/Engine/Plugins/Editor/CurveEditorTools/Source/CurveEditorTools/Private/CurveEditorToolCommands.h
--- a/Engine/Plugins/Editor/CurveEditorTools/Source/CurveEditorTools/Private/CurveEditorToolCommands.h
+++ b/Engine/Plugins/Editor/CurveEditorTools/Source/CurveEditorTools/Private/CurveEditorToolCommands.h
@@ -2,7 +2,6 @@
 
 #pragma once
 
-#include "CoreTypes.h"
 #include "Styling/AppStyle.h"
 #include "Framework/Commands/Commands.h"
 
diff --git a/Engine/Plugins/Editor/CurveEditorTools/Source/CurveEditorTools/Private/Filters/CurveEditorFFTFilter.cpp b/Engine/Plugins/Editor/CurveEditorTools/Source/CurveEditorTools/Private/Filters/CurveEditorFFTFilter.cpp
--- a/Engine/Plugins/Editor/CurveEditorTools/Source/CurveEditorTools/Private/Filters/CurveEditorFFTFilter.cpp
+++ b/Engine/Plugins/Editor/CurveEditorTools/Source/CurveEditorTools/Private/Filters/CurveEditorFFTFilter.cpp
@@ -2,7 +2,6 @@
 
 #include "CurveEditorFFTFilter.h"
 #include "DSP/PassiveFilter.h"
-#include "CurveEditor.h"
 #include "Filters/CurveEditorBakeFilter.h"
 
 #include UE_INLINE_GENERATED_CPP_BY_NAME(CurveEditorFFTFilter)
diff --git a/Engine/Plugins/Editor/CurveEditorTools/Source/CurveEditorTools/Private/Filters/CurveEditorFFTFilter.h b/Engine/Plugins/Editor/CurveEditorTools/Source/CurveEditorTools/Private/Filters/CurveEditorFFTFilter.h
--- a/Engine/Plugins/Editor/CurveEditorTools/Source/CurveEditorTools/Private/Filters/CurveEditorFFTFilter.h
+++ b/Engine/Plugins/Editor/CurveEditorTools/Source/CurveEditorTools/Private/Filters/CurveEditorFFTFilter.h
@@ -5,6 +5,9 @@
 #include "Filters/CurveEditorFilterBase.h"
 #include "CurveEditorFFTFilter.generated.h"
 
+struct FCurveModelID;
+struct FKeyHandleSet;
+
 
 UENUM(BlueprintType)
 enum class ECurveEditorFFTFilterClass : uint8
@@ -44,4 +47,4 @@ public:
 	/** The number of samples used to filter in the time domain. It maps how steep the roll off is for the filter. */
 	UPROPERTY(EditAnywhere, BlueprintReadWrite, meta=(UIMin=1, UIMax=8), Category = "Settings")
 	int32 Order;
-};
\ No newline at end of file
+};
diff --git a/Engine/Plugins/Editor/CurveEditorTools/Source/CurveEditorTools/Private/Tools/CurveEditorMultiScaleTool.cpp b/Engine/Plugins/Editor/CurveEditorTools/Source/CurveEditorTools/Private/Tools/CurveEditorMultiScaleTool.cpp
--- a/Engine/Plugins/Editor/CurveEditorTools/Source/CurveEditorTools/Private/Tools/CurveEditorMultiScaleTool.cpp
+++ b/Engine/Plugins/Editor/CurveEditorTools/Source/CurveEditorTools/Private/Tools/CurveEditorMultiScaleTool.cpp
@@ -4,21 +4,8 @@
 
 #include "CurveEditorToolCommands.h"
 #include "CurveEditor.h"
-#include "Framework/MultiBox/MultiBoxBuilder.h"
 #include "Framework/Commands/UICommandList.h"
-#include "Framework/Commands/UIAction.h"
-#include "ICurveEditorBounds.h"
-#include "Rendering/DrawElements.h"
-#include "CurveEditorScreenSpace.h"
-#include "Editor/Transactor.h"
-#include "ScopedTransaction.h"
-#include "CurveEditorSelection.h"
 #include "CurveModel.h"
-#include "CurveDataAbstraction.h"
-#include "Rendering/SlateRenderer.h"
-#include "Application/SlateApplicationBase.h"
-#include "Fonts/FontMeasure.h"
-#include "CurveEditorSnapMetrics.h"
 #include "SCurveEditorView.h"
 #include "SCurveEditorPanel.h"
 #include "Algo/MaxElement.h"
diff --git a/Engine/Plugins/Editor/CurveEditorTools/Source/CurveEditorTools/Private/Tools/CurveEditorMultiScaleTool.h b/Engine/Plugins/Editor/CurveEditorTools/Source/CurveEditorTools/Private/Tools/CurveEditorMultiScaleTool.h
--- a/Engine/Plugins/Editor/CurveEditorTools/Source/CurveEditorTools/Private/Tools/CurveEditorMultiScaleTool.h
+++ b/Engine/Plugins/Editor/CurveEditorTools/Source/CurveEditorTools/Private/Tools/CurveEditorMultiScaleTool.h
@@ -1,16 +1,10 @@
 // Copyright Epic Games, Inc. All Rights Reserved.
 
 #pragma once
-#include "CoreMinimal.h"
 #include "ICurveEditorToolExtension.h"
-#include "Layout/Geometry.h"
 #include "Framework/DelayedDrag.h"
-#include "Curves/KeyHandle.h"
 #include "CurveDataAbstraction.h"
 #include "CurveEditorTypes.h"
-#include "CurveEditorSelection.h"
-#include "Misc/EnumClassFlags.h"
-#include "CurveEditorSnapMetrics.h"
 #include "ScopedTransaction.h"
 #include "UObject/StructOnScope.h"
 
@@ -171,4 +165,4 @@ private:
 
 private:
 	FVector2D GetPivot(FCurveModel* InCurve, const TArray<FKeyPosition>& InKeyPositions) const;
-};
\ No newline at end of file
+};
diff --git a/Engine/Plugins/Editor/CurveEditorTools/Source/CurveEditorTools/Private/Tools/CurveEditorRetimeTool.cpp b/Engine/Plugins/Editor/CurveEditorTools/Source/CurveEditorTools/Private/Tools/CurveEditorRetimeTool.cpp
--- a/Engine/Plugins/Editor/CurveEditorTools/Source/CurveEditorTools/Private/Tools/CurveEditorRetimeTool.cpp
+++ b/Engine/Plugins/Editor/CurveEditorTools/Source/CurveEditorTools/Private/Tools/CurveEditorRetimeTool.cpp
@@ -3,24 +3,16 @@
 #include "CurveEditorRetimeTool.h"
 #include "CurveEditorToolCommands.h"
 #include "CurveEditor.h"
-#include "Framework/MultiBox/MultiBoxBuilder.h"
 #include "Framework/Commands/UICommandList.h"
-#include "Framework/Commands/UIAction.h"
 #include "ICurveEditorBounds.h"
 #include "Rendering/DrawElementPayloads.h"
-#include "Rendering/DrawElements.h"
-#include "CurveEditorScreenSpace.h"
-#include "Editor/Transactor.h"
 #include "ScopedTransaction.h"
-#include "CurveEditorSelection.h"
 #include "CurveModel.h"
-#include "CurveDataAbstraction.h"
 #include "EditorFontGlyphs.h"
 #include "Rendering/SlateRenderer.h"
-#include "Application/SlateApplicationBase.h"
 #include "Fonts/FontMeasure.h"
-#include "CurveEditorSnapMetrics.h"
 #include "SCurveEditorView.h"
+#include "Styling/CoreStyle.h"
 
 #include UE_INLINE_GENERATED_CPP_BY_NAME(CurveEditorRetimeTool)
 
diff --git a/Engine/Plugins/Editor/CurveEditorTools/Source/CurveEditorTools/Private/Tools/CurveEditorRetimeTool.h b/Engine/Plugins/Editor/CurveEditorTools/Source/CurveEditorTools/Private/Tools/CurveEditorRetimeTool.h
--- a/Engine/Plugins/Editor/CurveEditorTools/Source/CurveEditorTools/Private/Tools/CurveEditorRetimeTool.h
+++ b/Engine/Plugins/Editor/CurveEditorTools/Source/CurveEditorTools/Private/Tools/CurveEditorRetimeTool.h
@@ -1,13 +1,9 @@
 // Copyright Epic Games, Inc. All Rights Reserved.
 
 #pragma once
-#include "CoreMinimal.h"
 #include "ICurveEditorToolExtension.h"
-#include "Layout/Geometry.h"
 #include "Framework/DelayedDrag.h"
-#include "Curves/KeyHandle.h"
 #include "CurveDataAbstraction.h"
-#include "CurveEditorTypes.h"
 #include "CurveEditorSelection.h"
 #include "CurveEditorRetimeTool.generated.h"
 
@@ -134,4 +130,4 @@ private:
 
 	/** Cached user selection for keys. We clear this when the tool is opened and restore it when they leave */
 	TMap<FCurveModelID, FKeyHandleSet> CachedSelectionSet;
-};
\ No newline at end of file
+};
diff --git a/Engine/Plugins/Editor/CurveEditorTools/Source/CurveEditorTools/Private/Tools/CurveEditorTransformTool.cpp b/Engine/Plugins/Editor/CurveEditorTools/Source/CurveEditorTools/Private/Tools/CurveEditorTransformTool.cpp
--- a/Engine/Plugins/Editor/CurveEditorTools/Source/CurveEditorTools/Private/Tools/CurveEditorTransformTool.cpp
+++ b/Engine/Plugins/Editor/CurveEditorTools/Source/CurveEditorTools/Private/Tools/CurveEditorTransformTool.cpp
@@ -3,23 +3,12 @@
 #include "CurveEditorTransformTool.h"
 #include "CurveEditorToolCommands.h"
 #include "CurveEditor.h"
-#include "Framework/MultiBox/MultiBoxBuilder.h"
 #include "Framework/Commands/UICommandList.h"
-#include "Framework/Commands/UIAction.h"
 #include "Rendering/DrawElementPayloads.h"
-#include "Rendering/DrawElements.h"
-#include "Framework/DelayedDrag.h"
 #include "SCurveEditorView.h"
-#include "ScopedTransaction.h"
-#include "CurveEditorSelection.h"
 #include "SCurveEditorPanel.h"
 #include "CurveModel.h"
-#include "CurveEditorScreenSpace.h"
-#include "Containers/ArrayView.h"
-#include "CurveDataAbstraction.h"
-#include "CurveEditorSnapMetrics.h"
 #include "UObject/StructOnScope.h"
-#include "Algo/Transform.h"
 
 #include UE_INLINE_GENERATED_CPP_BY_NAME(CurveEditorTransformTool)
 
diff --git a/Engine/Plugins/Editor/CurveEditorTools/Source/CurveEditorTools/Private/Tools/CurveEditorTransformTool.h b/Engine/Plugins/Editor/CurveEditorTools/Source/CurveEditorTools/Private/Tools/CurveEditorTransformTool.h
--- a/Engine/Plugins/Editor/CurveEditorTools/Source/CurveEditorTools/Private/Tools/CurveEditorTransformTool.h
+++ b/Engine/Plugins/Editor/CurveEditorTools/Source/CurveEditorTools/Private/Tools/CurveEditorTransformTool.h
@@ -1,11 +1,9 @@
 // Copyright Epic Games, Inc. All Rights Reserved.
 
 #pragma once
-#include "Templates/SharedPointer.h"
+#include "CurveEditorTypes.h"
 #include "ICurveEditorToolExtension.h"
-#include "Misc/EnumClassFlags.h"
 #include "Framework/DelayedDrag.h"
-#include "Curves/KeyHandle.h"
 #include "CurveDataAbstraction.h"
 #include "ScopedTransaction.h"
 #include "CurveEditorSnapMetrics.h"
diff --git a/Engine/Plugins/Editor/DataValidation/Source/DataValidation/Private/DataValidationChangelist.cpp b/Engine/Plugins/Editor/DataValidation/Source/DataValidation/Private/DataValidationChangelist.cpp
--- a/Engine/Plugins/Editor/DataValidation/Source/DataValidation/Private/DataValidationChangelist.cpp
+++ b/Engine/Plugins/Editor/DataValidation/Source/DataValidation/Private/DataValidationChangelist.cpp
@@ -3,13 +3,10 @@
 #include "DataValidationChangelist.h"
 
 #include "Algo/AnyOf.h"
-#include "Algo/Transform.h"
-#include "Misc/Paths.h"
 #include "Misc/DataValidation.h"
 #include "AssetRegistry/AssetData.h"
-#include "AssetRegistry/AssetRegistryModule.h"
-#include "ISourceControlProvider.h"
 #include "ISourceControlModule.h"
+#include "Misc/PackageName.h"
 #include "SourceControlHelpers.h"
 #include "SourceControlOperations.h"
 #include "UncontrolledChangelistsModule.h"
diff --git a/Engine/Plugins/Editor/DataValidation/Source/DataValidation/Private/DataValidationManager.cpp b/Engine/Plugins/Editor/DataValidation/Source/DataValidation/Private/DataValidationManager.cpp
--- a/Engine/Plugins/Editor/DataValidation/Source/DataValidation/Private/DataValidationManager.cpp
+++ b/Engine/Plugins/Editor/DataValidation/Source/DataValidation/Private/DataValidationManager.cpp
@@ -2,6 +2,7 @@
 
 #include "DataValidationManager.h"
 
+#include "Misc/PackageName.h"
 #include "Modules/ModuleManager.h"
 #include "MessageLogModule.h"
 #include "Logging/MessageLog.h"
@@ -9,7 +10,6 @@
 #include "AssetRegistry/AssetRegistryModule.h"
 #include "Editor.h"
 
-#include "CoreGlobals.h"
 
 #include UE_INLINE_GENERATED_CPP_BY_NAME(DataValidationManager)
 
diff --git a/Engine/Plugins/Editor/DataValidation/Source/DataValidation/Private/DataValidationModule.cpp b/Engine/Plugins/Editor/DataValidation/Source/DataValidation/Private/DataValidationModule.cpp
--- a/Engine/Plugins/Editor/DataValidation/Source/DataValidation/Private/DataValidationModule.cpp
+++ b/Engine/Plugins/Editor/DataValidation/Source/DataValidation/Private/DataValidationModule.cpp
@@ -2,31 +2,26 @@
 
 #include "DataValidationModule.h"
 
-#include "UObject/Object.h"
+#include "AssetRegistry/ARFilter.h"
 #include "UObject/ObjectSaveContext.h"
-#include "UObject/SoftObjectPath.h"
-#include "GameFramework/HUD.h"
 
+#include "Editor.h"
 #include "Framework/Application/SlateApplication.h"
+#include "Internationalization/LocKeyFuncs.h"
 #include "ToolMenus.h"
-#include "Styling/AppStyle.h"
 #include "AssetRegistry/AssetRegistryModule.h"
 #include "Misc/MessageDialog.h"
 #include "DataValidationCommandlet.h"
-#include "LevelEditor.h"
-#include "Misc/FeedbackContext.h"
-#include "Misc/OutputDeviceConsole.h"
-#include "Modules/ModuleManager.h"
 
 #include "ContentBrowserMenuContexts.h"
-#include "ContentBrowserModule.h"
-#include "ContentBrowserDelegates.h"
 
-#include "WorkspaceMenuStructure.h"
-#include "WorkspaceMenuStructureModule.h"
 #include "EditorValidatorSubsystem.h"
 #include "ISettingsModule.h"
 #include "Algo/RemoveIf.h"
+#include "Misc/PackageName.h"
+#include "ToolMenu.h"
+#include "ToolMenuEntry.h"
+#include "ToolMenuSection.h"
 
 #define LOCTEXT_NAMESPACE "DataValidationModule"
 
diff --git a/Engine/Plugins/Editor/DataValidation/Source/DataValidation/Private/DirtyFilesChangelistValidator.cpp b/Engine/Plugins/Editor/DataValidation/Source/DataValidation/Private/DirtyFilesChangelistValidator.cpp
--- a/Engine/Plugins/Editor/DataValidation/Source/DataValidation/Private/DirtyFilesChangelistValidator.cpp
+++ b/Engine/Plugins/Editor/DataValidation/Source/DataValidation/Private/DirtyFilesChangelistValidator.cpp
@@ -2,12 +2,9 @@
 
 #include "DirtyFilesChangelistValidator.h"
 
-#include "Algo/Transform.h"
 #include "DataValidationChangelist.h"
 #include "FileHelpers.h"
-#include "ISourceControlChangelistState.h"
 #include "ISourceControlModule.h"
-#include "ISourceControlProvider.h"
 #include "Misc/Paths.h"
 
 #include UE_INLINE_GENERATED_CPP_BY_NAME(DirtyFilesChangelistValidator)
diff --git a/Engine/Plugins/Editor/DataValidation/Source/DataValidation/Private/EditorValidatorBase.cpp b/Engine/Plugins/Editor/DataValidation/Source/DataValidation/Private/EditorValidatorBase.cpp
--- a/Engine/Plugins/Editor/DataValidation/Source/DataValidation/Private/EditorValidatorBase.cpp
+++ b/Engine/Plugins/Editor/DataValidation/Source/DataValidation/Private/EditorValidatorBase.cpp
@@ -2,15 +2,8 @@
 
 #include "EditorValidatorBase.h"
 
-#include "Editor.h"
-#include "AssetToolsModule.h"
-#include "ObjectTools.h"
-#include "Dialogs/Dialogs.h"
 #include "EditorValidatorSubsystem.h"
-#include "Logging/MessageLog.h"
-#include "Misc/UObjectToken.h"
-#include "Logging/TokenizedMessage.h"
-#include "DataValidationModule.h"
+#include "Trace/Trace.inl"
 
 #include UE_INLINE_GENERATED_CPP_BY_NAME(EditorValidatorBase)
 
diff --git a/Engine/Plugins/Editor/DataValidation/Source/DataValidation/Private/EditorValidatorSubsystem.cpp b/Engine/Plugins/Editor/DataValidation/Source/DataValidation/Private/EditorValidatorSubsystem.cpp
--- a/Engine/Plugins/Editor/DataValidation/Source/DataValidation/Private/EditorValidatorSubsystem.cpp
+++ b/Engine/Plugins/Editor/DataValidation/Source/DataValidation/Private/EditorValidatorSubsystem.cpp
@@ -2,26 +2,23 @@
 
 #include "EditorValidatorSubsystem.h"
 
+#include "AssetRegistry/ARFilter.h"
 #include "Editor.h"
-#include "AssetToolsModule.h"
-#include "ObjectTools.h"
 #include "AssetRegistry/AssetRegistryModule.h"
+#include "Blueprint/BlueprintSupport.h"
 #include "IDirectoryWatcher.h"
 #include "DirectoryWatcherModule.h"
 #include "EditorUtilityBlueprint.h"
-#include "Settings/EditorLoadingSavingSettings.h"
-#include "UnrealEdGlobals.h"
-#include "Editor/UnrealEdEngine.h"
-#include "CookOnTheSide/CookOnTheFlyServer.h"
+#include "EditorValidatorBase.h"
 #include "Logging/MessageLog.h"
+#include "Misc/App.h"
 #include "Misc/ScopedSlowTask.h"
 #include "Misc/DataValidation.h"
-#include "AssetRegistry/AssetData.h"
 #include "ISourceControlModule.h"
 #include "DataValidationChangelist.h"
-#include "DataValidationModule.h"
 #include "Engine/Level.h"
 #include "Editor.h"
+#include "Misc/PackageName.h"
 
 #include UE_INLINE_GENERATED_CPP_BY_NAME(EditorValidatorSubsystem)
 
diff --git a/Engine/Plugins/Editor/DataValidation/Source/DataValidation/Private/EditorValidator_Localization.cpp b/Engine/Plugins/Editor/DataValidation/Source/DataValidation/Private/EditorValidator_Localization.cpp
--- a/Engine/Plugins/Editor/DataValidation/Source/DataValidation/Private/EditorValidator_Localization.cpp
+++ b/Engine/Plugins/Editor/DataValidation/Source/DataValidation/Private/EditorValidator_Localization.cpp
@@ -2,9 +2,11 @@
 
 #include "EditorValidator_Localization.h"
 
+#include "AssetRegistry/AssetData.h"
 #include "AssetRegistry/AssetRegistryModule.h"
 #include "IAssetTools.h"
 #include "AssetToolsModule.h"
+#include "Misc/PackageName.h"
 #include "Misc/Paths.h"
 #include "HAL/FileManager.h"
 #include "Internationalization/Culture.h"
diff --git a/Engine/Plugins/Editor/DataValidation/Source/DataValidation/Private/WorldPartitionChangelistValidator.cpp b/Engine/Plugins/Editor/DataValidation/Source/DataValidation/Private/WorldPartitionChangelistValidator.cpp
--- a/Engine/Plugins/Editor/DataValidation/Source/DataValidation/Private/WorldPartitionChangelistValidator.cpp
+++ b/Engine/Plugins/Editor/DataValidation/Source/DataValidation/Private/WorldPartitionChangelistValidator.cpp
@@ -1,24 +1,23 @@
 // Copyright Epic Games, Inc. All Rights Reserved.
 
 #include "WorldPartitionChangelistValidator.h"
+#include "AssetRegistry/ARFilter.h"
 #include "DataValidationChangelist.h"
 
 #include "AssetRegistry/AssetData.h"
 #include "AssetRegistry/AssetRegistryModule.h"
-#include "ISourceControlProvider.h"
 #include "ISourceControlModule.h"
+#include "Misc/PackageName.h"
 #include "SourceControlHelpers.h"
 
 #include "Engine/Level.h"
-#include "Engine/World.h"
 #include "Misc/PathViews.h"
-#include "WorldPartition/WorldPartitionActorDesc.h"
 #include "WorldPartition/WorldPartitionActorDescUtils.h"
-#include "WorldPartition/ActorDescContainer.h"
 #include "WorldPartition/WorldPartition.h"
 #include "WorldPartition/DataLayer/DataLayerInstanceWithAsset.h"
 #include "WorldPartition/DataLayer/WorldDataLayers.h"
 #include "WorldPartition/ContentBundle/ContentBundlePaths.h"
+#include "WorldPartition/WorldPartitionActorDescView.h"
 
 #include UE_INLINE_GENERATED_CPP_BY_NAME(WorldPartitionChangelistValidator)
 
diff --git a/Engine/Plugins/Editor/EditorDebugTools/Source/EditorDebugTools/Private/EditorDebugTools.cpp b/Engine/Plugins/Editor/EditorDebugTools/Source/EditorDebugTools/Private/EditorDebugTools.cpp
--- a/Engine/Plugins/Editor/EditorDebugTools/Source/EditorDebugTools/Private/EditorDebugTools.cpp
+++ b/Engine/Plugins/Editor/EditorDebugTools/Source/EditorDebugTools/Private/EditorDebugTools.cpp
@@ -2,11 +2,7 @@
 
 #include "EditorDebugTools.h"
 #include "EditorDebugToolsStyle.h"
-#include "LevelEditor.h"
 #include "Widgets/Docking/SDockTab.h"
-#include "Widgets/Layout/SBox.h"
-#include "Widgets/Text/STextBlock.h"
-#include "ToolMenus.h"
 #include "SDebugPanel.h"
 #include "GammaUIPanel.h"
 #include "SModuleUI.h"
@@ -83,4 +79,4 @@ void FEditorDebugToolsModule::ShutdownModule()
 
 #undef LOCTEXT_NAMESPACE
 	
-IMPLEMENT_MODULE(FEditorDebugToolsModule, EditorDebugTools)
\ No newline at end of file
+IMPLEMENT_MODULE(FEditorDebugToolsModule, EditorDebugTools)
diff --git a/Engine/Plugins/Editor/EditorDebugTools/Source/EditorDebugTools/Private/EditorDebugToolsStyle.cpp b/Engine/Plugins/Editor/EditorDebugTools/Source/EditorDebugTools/Private/EditorDebugToolsStyle.cpp
--- a/Engine/Plugins/Editor/EditorDebugTools/Source/EditorDebugTools/Private/EditorDebugToolsStyle.cpp
+++ b/Engine/Plugins/Editor/EditorDebugTools/Source/EditorDebugTools/Private/EditorDebugToolsStyle.cpp
@@ -1,10 +1,10 @@
 // Copyright Epic Games, Inc. All Rights Reserved.
 
 #include "EditorDebugToolsStyle.h"
+#include "Brushes/SlateImageBrush.h"
 #include "Styling/SlateStyleRegistry.h"
-#include "Framework/Application/SlateApplication.h"
 #include "Interfaces/IPluginManager.h"
-#include "Styling/SlateStyleMacros.h"
+#include "Styling/SlateStyle.h"
 
 TSharedPtr< FSlateStyleSet > FEditorDebugToolsStyle::StyleInstance = NULL;
 
diff --git a/Engine/Plugins/Editor/EditorDebugTools/Source/EditorDebugTools/Private/GammaUIPanel.cpp b/Engine/Plugins/Editor/EditorDebugTools/Source/EditorDebugTools/Private/GammaUIPanel.cpp
--- a/Engine/Plugins/Editor/EditorDebugTools/Source/EditorDebugTools/Private/GammaUIPanel.cpp
+++ b/Engine/Plugins/Editor/EditorDebugTools/Source/EditorDebugTools/Private/GammaUIPanel.cpp
@@ -1,13 +1,8 @@
 // Copyright Epic Games, Inc. All Rights Reserved.
 
 #include "GammaUIPanel.h"
-#include "EngineGlobals.h"
 #include "Engine/Engine.h"
-#include "Widgets/SBoxPanel.h"
-#include "Widgets/Images/SImage.h"
-#include "Widgets/Text/STextBlock.h"
 #include "Widgets/Input/SSpinBox.h"
-#include "Styling/AppStyle.h"
 #include "EditorDebugToolsStyle.h"
 
 void SGammaUIPanel::Construct(const SGammaUIPanel::FArguments& InArgs)
diff --git a/Engine/Plugins/Editor/EditorDebugTools/Source/EditorDebugTools/Private/GammaUIPanel.h b/Engine/Plugins/Editor/EditorDebugTools/Source/EditorDebugTools/Private/GammaUIPanel.h
--- a/Engine/Plugins/Editor/EditorDebugTools/Source/EditorDebugTools/Private/GammaUIPanel.h
+++ b/Engine/Plugins/Editor/EditorDebugTools/Source/EditorDebugTools/Private/GammaUIPanel.h
@@ -2,8 +2,6 @@
 
 #pragma once
 
-#include "CoreMinimal.h"
-#include "Widgets/DeclarativeSyntaxSupport.h"
 #include "Widgets/SCompoundWidget.h"
 
 class SGammaUIPanel : public SCompoundWidget
diff --git a/Engine/Plugins/Editor/EditorDebugTools/Source/EditorDebugTools/Private/SDebugPanel.cpp b/Engine/Plugins/Editor/EditorDebugTools/Source/EditorDebugTools/Private/SDebugPanel.cpp
--- a/Engine/Plugins/Editor/EditorDebugTools/Source/EditorDebugTools/Private/SDebugPanel.cpp
+++ b/Engine/Plugins/Editor/EditorDebugTools/Source/EditorDebugTools/Private/SDebugPanel.cpp
@@ -1,11 +1,11 @@
 // Copyright Epic Games, Inc. All Rights Reserved.
 
 #include "SDebugPanel.h"
+#include "Rendering/SlateRenderer.h"
 #include "Widgets/Testing/SStarshipSuite.h"
 #include "Widgets/Testing/STestSuite.h"
 #include "Widgets/SBoxPanel.h"
 #include "Widgets/Input/SButton.h"
-#include "Styling/CoreStyle.h"
 #include "Modules/ModuleManager.h"
 #include "Framework/Application/SlateApplication.h"
 
@@ -110,4 +110,4 @@ FReply SDebugPanel::OnTestSuiteClicked()
 	}
 #endif
 	return FReply::Handled();
-}
\ No newline at end of file
+}
diff --git a/Engine/Plugins/Editor/EditorDebugTools/Source/EditorDebugTools/Private/SModuleUI.cpp b/Engine/Plugins/Editor/EditorDebugTools/Source/EditorDebugTools/Private/SModuleUI.cpp
--- a/Engine/Plugins/Editor/EditorDebugTools/Source/EditorDebugTools/Private/SModuleUI.cpp
+++ b/Engine/Plugins/Editor/EditorDebugTools/Source/EditorDebugTools/Private/SModuleUI.cpp
@@ -1,14 +1,13 @@
 // Copyright Epic Games, Inc. All Rights Reserved.
 
 #include "SModuleUI.h"
+#include "Framework/Views/TableViewMetadata.h"
 #include "HAL/PlatformFileManager.h"
 #include "Misc/FeedbackContext.h"
 #include "GenericPlatform/GenericPlatformFile.h"
-#include "Modules/ModuleManager.h"
-#include "EngineGlobals.h"
 #include "Engine/Engine.h"
+#include "Misc/Paths.h"
 #include "SlateOptMacros.h"
-#include "Widgets/Text/STextBlock.h"
 #include "Widgets/Input/SButton.h"
 #include "Widgets/Views/SListView.h"
 #include "Misc/HotReloadInterface.h"
diff --git a/Engine/Plugins/Editor/EditorDebugTools/Source/EditorDebugTools/Private/SModuleUI.h b/Engine/Plugins/Editor/EditorDebugTools/Source/EditorDebugTools/Private/SModuleUI.h
--- a/Engine/Plugins/Editor/EditorDebugTools/Source/EditorDebugTools/Private/SModuleUI.h
+++ b/Engine/Plugins/Editor/EditorDebugTools/Source/EditorDebugTools/Private/SModuleUI.h
@@ -2,14 +2,11 @@
 
 #pragma once
 
-#include "CoreMinimal.h"
-#include "SlateFwd.h"
-#include "Layout/Visibility.h"
-#include "Input/Reply.h"
-#include "Widgets/DeclarativeSyntaxSupport.h"
 #include "Widgets/SCompoundWidget.h"
-#include "Widgets/Views/STableViewBase.h"
-#include "Widgets/Views/STableRow.h"
+
+class ITableRow;
+class SSearchBox;
+class STableViewBase;
 
 template <typename ItemType> class SListView;
 
diff --git a/Engine/Plugins/Editor/EngineAssetDefinitions/Source/Private/AssetDefinition_Actor.h b/Engine/Plugins/Editor/EngineAssetDefinitions/Source/Private/AssetDefinition_Actor.h
--- a/Engine/Plugins/Editor/EngineAssetDefinitions/Source/Private/AssetDefinition_Actor.h
+++ b/Engine/Plugins/Editor/EngineAssetDefinitions/Source/Private/AssetDefinition_Actor.h
@@ -2,12 +2,15 @@
 
 #pragma once
 
-#include "CoreMinimal.h"
 #include "GameFramework/Actor.h"
 #include "AssetDefinitionDefault.h"
 
 #include "AssetDefinition_Actor.generated.h"
 
+enum class EAssetCommandResult : uint8;
+struct FAssetCategoryPath;
+struct FAssetOpenArgs;
+
 UCLASS()
 class UAssetDefinition_Actor : public UAssetDefinitionDefault
 {
diff --git a/Engine/Plugins/Editor/EngineAssetDefinitions/Source/Private/AssetDefinition_ActorFolder.h b/Engine/Plugins/Editor/EngineAssetDefinitions/Source/Private/AssetDefinition_ActorFolder.h
--- a/Engine/Plugins/Editor/EngineAssetDefinitions/Source/Private/AssetDefinition_ActorFolder.h
+++ b/Engine/Plugins/Editor/EngineAssetDefinitions/Source/Private/AssetDefinition_ActorFolder.h
@@ -7,6 +7,8 @@
 
 #include "AssetDefinition_ActorFolder.generated.h"
 
+struct FAssetCategoryPath;
+
 UCLASS()
 class UAssetDefinition_ActorFolder : public UAssetDefinitionDefault
 {
diff --git a/Engine/Plugins/Editor/EngineAssetDefinitions/Source/Private/AssetDefinition_Blueprint.cpp b/Engine/Plugins/Editor/EngineAssetDefinitions/Source/Private/AssetDefinition_Blueprint.cpp
--- a/Engine/Plugins/Editor/EngineAssetDefinitions/Source/Private/AssetDefinition_Blueprint.cpp
+++ b/Engine/Plugins/Editor/EngineAssetDefinitions/Source/Private/AssetDefinition_Blueprint.cpp
@@ -4,16 +4,23 @@
 
 #include "AssetDefinitionRegistry.h"
 #include "AssetToolsModule.h"
+#include "Blueprint/BlueprintSupport.h"
 #include "BlueprintEditorModule.h"
 #include "ContentBrowserMenuContexts.h"
 #include "ContentBrowserModule.h"
+#include "IAssetTools.h"
 #include "ToolMenus.h"
 #include "Factories/BlueprintFactory.h"
+#include "Internationalization/LocKeyFuncs.h"
 #include "Kismet2/KismetEditorUtilities.h"
 #include "IContentBrowserSingleton.h"
 #include "Algo/AllOf.h"
 #include "Logging/MessageLog.h"
+#include "Misc/DelayedAutoRegister.h"
 #include "Misc/MessageDialog.h"
+#include "Misc/PackageName.h"
+#include "ToolMenu.h"
+#include "ToolMenuSection.h"
 
 #define LOCTEXT_NAMESPACE "UAssetDefinition_Blueprint"
 
diff --git a/Engine/Plugins/Editor/EngineAssetDefinitions/Source/Private/AssetDefinition_Curve.cpp b/Engine/Plugins/Editor/EngineAssetDefinitions/Source/Private/AssetDefinition_Curve.cpp
--- a/Engine/Plugins/Editor/EngineAssetDefinitions/Source/Private/AssetDefinition_Curve.cpp
+++ b/Engine/Plugins/Editor/EngineAssetDefinitions/Source/Private/AssetDefinition_Curve.cpp
@@ -4,15 +4,18 @@
 
 #include "ContentBrowserMenuContexts.h"
 #include "EditorFramework/AssetImportData.h"
-#include "ICurveAssetEditor.h"
 #include "CurveAssetEditorModule.h"
 #include "DesktopPlatformModule.h"
 #include "IAssetTools.h"
 #include "Framework/Application/SlateApplication.h"
+#include "IDesktopPlatform.h"
 #include "ToolMenuSection.h"
+#include "Misc/DelayedAutoRegister.h"
 #include "Misc/FileHelper.h"
 #include "Misc/MessageDialog.h"
+#include "Misc/Paths.h"
 #include "Settings/EditorLoadingSavingSettings.h"
+#include "ToolMenu.h"
 #include "ToolMenus.h"
 
 #define LOCTEXT_NAMESPACE "UAssetDefinition_Curve"
diff --git a/Engine/Plugins/Editor/EngineAssetDefinitions/Source/Private/AssetDefinition_Curve.h b/Engine/Plugins/Editor/EngineAssetDefinitions/Source/Private/AssetDefinition_Curve.h
--- a/Engine/Plugins/Editor/EngineAssetDefinitions/Source/Private/AssetDefinition_Curve.h
+++ b/Engine/Plugins/Editor/EngineAssetDefinitions/Source/Private/AssetDefinition_Curve.h
@@ -2,12 +2,15 @@
 
 #pragma once
 
-#include "CoreMinimal.h"
 #include "AssetDefinitionDefault.h"
 #include "Curves/CurveBase.h"
 
 #include "AssetDefinition_Curve.generated.h"
 
+enum class EAssetCommandResult : uint8;
+struct FAssetDiffArgs;
+struct FAssetOpenArgs;
+
 UCLASS()
 class UAssetDefinition_Curve : public UAssetDefinitionDefault
 {
diff --git a/Engine/Plugins/Editor/EngineAssetDefinitions/Source/Private/AssetDefinition_DataAsset.cpp b/Engine/Plugins/Editor/EngineAssetDefinitions/Source/Private/AssetDefinition_DataAsset.cpp
--- a/Engine/Plugins/Editor/EngineAssetDefinitions/Source/Private/AssetDefinition_DataAsset.cpp
+++ b/Engine/Plugins/Editor/EngineAssetDefinitions/Source/Private/AssetDefinition_DataAsset.cpp
@@ -2,16 +2,17 @@
 
 #include "AssetDefinition_DataAsset.h"
 
-#include "AssetToolsModule.h"
 #include "ClassViewerFilter.h"
-#include "ClassViewerModule.h"
 #include "ContentBrowserMenuContexts.h"
+#include "Misc/DelayedAutoRegister.h"
 #include "ObjectTools.h"
 #include "SDetailsDiff.h"
+#include "ToolMenu.h"
 #include "ToolMenus.h"
 #include "Kismet2/SClassPickerDialog.h"
 #include "Engine/Engine.h"
 #include "Misc/AssetFilterData.h"
+#include "ToolMenuSection.h"
 
 #define LOCTEXT_NAMESPACE "UAssetDefinition_DataAsset"
 
diff --git a/Engine/Plugins/Editor/EngineAssetDefinitions/Source/Private/AssetDefinition_Font.cpp b/Engine/Plugins/Editor/EngineAssetDefinitions/Source/Private/AssetDefinition_Font.cpp
--- a/Engine/Plugins/Editor/EngineAssetDefinitions/Source/Private/AssetDefinition_Font.cpp
+++ b/Engine/Plugins/Editor/EngineAssetDefinitions/Source/Private/AssetDefinition_Font.cpp
@@ -3,10 +3,12 @@
 #include "AssetDefinition_Font.h"
 
 #include "ContentBrowserMenuContexts.h"
+#include "Misc/DelayedAutoRegister.h"
 #include "ToolMenus.h"
-#include "Styling/AppStyle.h"
 #include "EditorReimportHandler.h"
 #include "FontEditorModule.h"
+#include "ToolMenu.h"
+#include "ToolMenuSection.h"
 
 #define LOCTEXT_NAMESPACE "UAssetDefinition_Font"
 
diff --git a/Engine/Plugins/Editor/EngineAssetDefinitions/Source/Private/AssetDefinition_FontFace.cpp b/Engine/Plugins/Editor/EngineAssetDefinitions/Source/Private/AssetDefinition_FontFace.cpp
--- a/Engine/Plugins/Editor/EngineAssetDefinitions/Source/Private/AssetDefinition_FontFace.cpp
+++ b/Engine/Plugins/Editor/EngineAssetDefinitions/Source/Private/AssetDefinition_FontFace.cpp
@@ -4,10 +4,12 @@
 
 #include "ContentBrowserMenuContexts.h"
 #include "FontEditorModule.h"
+#include "Misc/DelayedAutoRegister.h"
 #include "ToolMenus.h"
-#include "Styling/AppStyle.h"
 #include "EditorReimportHandler.h"
+#include "ToolMenu.h"
 #include "Toolkits/SimpleAssetEditor.h"
+#include "ToolMenuSection.h"
 
 #define LOCTEXT_NAMESPACE "AssetDefinition_FontFace"
 
diff --git a/Engine/Plugins/Editor/EngineAssetDefinitions/Source/Private/AssetDefinition_ObjectLibrary.h b/Engine/Plugins/Editor/EngineAssetDefinitions/Source/Private/AssetDefinition_ObjectLibrary.h
--- a/Engine/Plugins/Editor/EngineAssetDefinitions/Source/Private/AssetDefinition_ObjectLibrary.h
+++ b/Engine/Plugins/Editor/EngineAssetDefinitions/Source/Private/AssetDefinition_ObjectLibrary.h
@@ -2,8 +2,6 @@
 
 #pragma once
 
-#include "CoreMinimal.h"
-#include "AssetTypeActions_Base.h"
 #include "Engine/ObjectLibrary.h"
 #include "AssetDefinitionDefault.h"
 
diff --git a/Engine/Plugins/Editor/EngineAssetDefinitions/Source/Private/AssetDefinition_PoseAsset.cpp b/Engine/Plugins/Editor/EngineAssetDefinitions/Source/Private/AssetDefinition_PoseAsset.cpp
--- a/Engine/Plugins/Editor/EngineAssetDefinitions/Source/Private/AssetDefinition_PoseAsset.cpp
+++ b/Engine/Plugins/Editor/EngineAssetDefinitions/Source/Private/AssetDefinition_PoseAsset.cpp
@@ -1,14 +1,10 @@
 // Copyright Epic Games, Inc. All Rights Reserved.
 
 #include "AssetDefinition_PoseAsset.h"
+#include "Animation/Skeleton.h"
 #include "ToolMenuSection.h"
-#include "Animation/PoseAsset.h"
 #include "Animation/AnimSequence.h"
-#include "Misc/MessageDialog.h"
-#include "Styling/AppStyle.h"
-#include "ToolMenus.h"
 #include "ScopedTransaction.h"
-#include "Framework/Application/SlateApplication.h"
 
 #define LOCTEXT_NAMESPACE "AssetTypeActions"
 
diff --git a/Engine/Plugins/Editor/EngineAssetDefinitions/Source/Private/AssetDefinition_Redirector.cpp b/Engine/Plugins/Editor/EngineAssetDefinitions/Source/Private/AssetDefinition_Redirector.cpp
--- a/Engine/Plugins/Editor/EngineAssetDefinitions/Source/Private/AssetDefinition_Redirector.cpp
+++ b/Engine/Plugins/Editor/EngineAssetDefinitions/Source/Private/AssetDefinition_Redirector.cpp
@@ -4,7 +4,11 @@
 
 #include "AssetToolsModule.h"
 #include "ContentBrowserMenuContexts.h"
+#include "IAssetTools.h"
 #include "ToolMenus.h"
+#include "Misc/DelayedAutoRegister.h"
+#include "ToolMenu.h"
+#include "ToolMenuSection.h"
 
 #define LOCTEXT_NAMESPACE "UAssetDefinition_Redirector"
 
diff --git a/Engine/Plugins/Editor/EngineAssetDefinitions/Source/Private/AssetDefinition_SlateWidgetStyle.cpp b/Engine/Plugins/Editor/EngineAssetDefinitions/Source/Private/AssetDefinition_SlateWidgetStyle.cpp
--- a/Engine/Plugins/Editor/EngineAssetDefinitions/Source/Private/AssetDefinition_SlateWidgetStyle.cpp
+++ b/Engine/Plugins/Editor/EngineAssetDefinitions/Source/Private/AssetDefinition_SlateWidgetStyle.cpp
@@ -1,7 +1,6 @@
 // Copyright Epic Games, Inc. All Rights Reserved.
 
 #include "AssetDefinition_SlateWidgetStyle.h"
-#include "ToolMenus.h"
 #include "Toolkits/SimpleAssetEditor.h"
 
 #define LOCTEXT_NAMESPACE "UAssetDefinition_SlateWidgetStyle"
diff --git a/Engine/Plugins/Editor/EngineAssetDefinitions/Source/Private/AssetDefinition_StaticMesh.cpp b/Engine/Plugins/Editor/EngineAssetDefinitions/Source/Private/AssetDefinition_StaticMesh.cpp
--- a/Engine/Plugins/Editor/EngineAssetDefinitions/Source/Private/AssetDefinition_StaticMesh.cpp
+++ b/Engine/Plugins/Editor/EngineAssetDefinitions/Source/Private/AssetDefinition_StaticMesh.cpp
@@ -2,23 +2,19 @@
 
 #include "AssetDefinition_StaticMesh.h"
 
-#include "IAssetTools.h"
-#include "Editor.h"
+#include "Algo/Accumulate.h"
 #include "EditorSupportDelegates.h"
+#include "Engine/StaticMeshSourceData.h"
 #include "StaticMeshEditorModule.h"
-#include "EditorFramework/AssetImportData.h"
 #include "FbxMeshUtils.h"
-#include "Framework/Notifications/NotificationManager.h"
+#include "Misc/DelayedAutoRegister.h"
 #include "Misc/MessageDialog.h"
-#include "Misc/ScopedSlowTask.h"
 #include "StaticMeshCompiler.h"
-#include "Styling/AppStyle.h"
 #include "ThumbnailRendering/SceneThumbnailInfo.h"
+#include "ToolMenu.h"
 #include "ToolMenus.h"
-#include "Widgets/Notifications/SNotificationList.h"
 #include "ContentBrowserMenuContexts.h"
-#include "Algo/AllOf.h"
-#include "Algo/Compare.h"
+#include "ToolMenuSection.h"
 
 #define LOCTEXT_NAMESPACE "UAssetDefinition_StaticMesh"
 
diff --git a/Engine/Plugins/Editor/EngineAssetDefinitions/Source/Private/AssetDefinition_Texture.cpp b/Engine/Plugins/Editor/EngineAssetDefinitions/Source/Private/AssetDefinition_Texture.cpp
--- a/Engine/Plugins/Editor/EngineAssetDefinitions/Source/Private/AssetDefinition_Texture.cpp
+++ b/Engine/Plugins/Editor/EngineAssetDefinitions/Source/Private/AssetDefinition_Texture.cpp
@@ -2,24 +2,24 @@
 
 #include "AssetDefinition_Texture.h"
 
+#include "AssetRegistry/AssetIdentifier.h"
 #include "ContentBrowserMenuContexts.h"
 #include "Engine/Texture2D.h"
+#include "Misc/DelayedAutoRegister.h"
 #include "Misc/PackageName.h"
-#include "Styling/AppStyle.h"
 #include "Materials/Material.h"
 #include "Materials/MaterialFunction.h"
-#include "Materials/MaterialInstance.h"
 #include "Factories/MaterialFactoryNew.h"
 #include "Interfaces/ITextureEditorModule.h"
 #include "IContentBrowserSingleton.h"
 #include "ContentBrowserModule.h"
 #include "Factories/SubUVAnimationFactory.h"
 #include "Particles/SubUVAnimation.h"
-#include "AssetRegistry/AssetRegistryModule.h"
 
+#include "ToolMenu.h"
 #include "ToolMenus.h"
-#include "Misc/FeedbackContext.h"
 #include "IAssetTools.h"
+#include "ToolMenuSection.h"
 #include "VirtualTexturingEditorModule.h"
 #include "Algo/AnyOf.h"
 #include "AssetRegistry/AssetRegistryHelpers.h"
diff --git a/Engine/Plugins/Editor/EngineAssetDefinitions/Source/Private/AssetDefinition_TextureCube.cpp b/Engine/Plugins/Editor/EngineAssetDefinitions/Source/Private/AssetDefinition_TextureCube.cpp
--- a/Engine/Plugins/Editor/EngineAssetDefinitions/Source/Private/AssetDefinition_TextureCube.cpp
+++ b/Engine/Plugins/Editor/EngineAssetDefinitions/Source/Private/AssetDefinition_TextureCube.cpp
@@ -2,16 +2,17 @@
 
 #include "AssetDefinition_TextureCube.h"
 
-#include "AssetToolsModule.h"
 #include "ContentBrowserMenuContexts.h"
+#include "Misc/DelayedAutoRegister.h"
 #include "ToolMenus.h"
 #include "Misc/PackageName.h"
-#include "Styling/AppStyle.h"
 #include "Factories/TextureCubeArrayFactory.h"
 #include "Engine/TextureCubeArray.h"
 #include "IAssetTools.h"
 #include "IContentBrowserSingleton.h"
 #include "ContentBrowserModule.h"
+#include "ToolMenu.h"
+#include "ToolMenuSection.h"
 
 #define LOCTEXT_NAMESPACE "UAssetDefinition_TextureCubeArray"
 
diff --git a/Engine/Plugins/Editor/EngineAssetDefinitions/Source/Private/AssetDefinition_TextureRenderTarget.cpp b/Engine/Plugins/Editor/EngineAssetDefinitions/Source/Private/AssetDefinition_TextureRenderTarget.cpp
--- a/Engine/Plugins/Editor/EngineAssetDefinitions/Source/Private/AssetDefinition_TextureRenderTarget.cpp
+++ b/Engine/Plugins/Editor/EngineAssetDefinitions/Source/Private/AssetDefinition_TextureRenderTarget.cpp
@@ -2,10 +2,9 @@
 
 #include "AssetDefinition_TextureRenderTarget.h"
 
-#include "AssetToolsModule.h"
 #include "ContentBrowserMenuContexts.h"
+#include "IAssetTools.h"
 #include "ToolMenus.h"
-#include "Styling/AppStyle.h"
 #include "Engine/Texture2D.h"
 #include "Engine/Texture2DArray.h"
 #include "Engine/TextureCube.h"
@@ -15,6 +14,9 @@
 #include "Engine/TextureRenderTargetCube.h"
 #include "Engine/TextureRenderTargetVolume.h"
 #include "AssetRegistry/AssetRegistryModule.h"
+#include "Misc/DelayedAutoRegister.h"
+#include "ToolMenu.h"
+#include "ToolMenuSection.h"
 
 #define LOCTEXT_NAMESPACE "UAssetDefinition_TextureRenderTarget"
 
diff --git a/Engine/Plugins/Editor/EngineAssetDefinitions/Source/Private/AssetDefinition_UserDefinedEnum.cpp b/Engine/Plugins/Editor/EngineAssetDefinitions/Source/Private/AssetDefinition_UserDefinedEnum.cpp
--- a/Engine/Plugins/Editor/EngineAssetDefinitions/Source/Private/AssetDefinition_UserDefinedEnum.cpp
+++ b/Engine/Plugins/Editor/EngineAssetDefinitions/Source/Private/AssetDefinition_UserDefinedEnum.cpp
@@ -2,7 +2,7 @@
 
 #include "AssetDefinition_UserDefinedEnum.h"
 #include "BlueprintEditorModule.h"
-#include "AssetRegistry/AssetData.h"
+#include "Modules/ModuleManager.h"
 
 #define LOCTEXT_NAMESPACE "UAssetDefinition_UserDefinedEnum"
 
diff --git a/Engine/Plugins/Editor/EngineAssetDefinitions/Source/Private/AssetDefinition_VolumeTexture.h b/Engine/Plugins/Editor/EngineAssetDefinitions/Source/Private/AssetDefinition_VolumeTexture.h
--- a/Engine/Plugins/Editor/EngineAssetDefinitions/Source/Private/AssetDefinition_VolumeTexture.h
+++ b/Engine/Plugins/Editor/EngineAssetDefinitions/Source/Private/AssetDefinition_VolumeTexture.h
@@ -2,9 +2,8 @@
 
 #pragma once
 
-#include "CoreMinimal.h"
+#include "AssetDefinition_Texture.h"
 #include "Engine/VolumeTexture.h"
-#include "AssetDefinition_TextureRenderTarget.h"
 
 #include "AssetDefinition_VolumeTexture.generated.h"
 
diff --git a/Engine/Plugins/Editor/EngineAssetDefinitions/Source/Private/AssetDefinition_World.cpp b/Engine/Plugins/Editor/EngineAssetDefinitions/Source/Private/AssetDefinition_World.cpp
--- a/Engine/Plugins/Editor/EngineAssetDefinitions/Source/Private/AssetDefinition_World.cpp
+++ b/Engine/Plugins/Editor/EngineAssetDefinitions/Source/Private/AssetDefinition_World.cpp
@@ -3,6 +3,7 @@
 #include "AssetDefinition_World.h"
 
 #include "FileHelpers.h"
+#include "Misc/PackageName.h"
 #include "ThumbnailRendering/WorldThumbnailInfo.h"
 #include "Misc/MessageDialog.h"
 #include "Engine/Level.h"
diff --git a/Engine/Plugins/Editor/EngineAssetDefinitions/Source/Private/EngineAssetDefinitionsModule.cpp b/Engine/Plugins/Editor/EngineAssetDefinitions/Source/Private/EngineAssetDefinitionsModule.cpp
--- a/Engine/Plugins/Editor/EngineAssetDefinitions/Source/Private/EngineAssetDefinitionsModule.cpp
+++ b/Engine/Plugins/Editor/EngineAssetDefinitions/Source/Private/EngineAssetDefinitionsModule.cpp
@@ -1,8 +1,5 @@
 ﻿// Copyright Epic Games, Inc. All Rights Reserved.
 
-#include "Modules/ModuleManager.h"
-#include "AssetToolsModule.h"
-#include "IAssetTypeActions.h"
 #include "AssetDefinition_PoseAsset.h"
 
 #define LOCTEXT_NAMESPACE "FEngineAssetDefinitionsModule"
diff --git a/Engine/Plugins/Editor/EngineAssetDefinitions/Source/Private/Material/AssetDefinition_Material.h b/Engine/Plugins/Editor/EngineAssetDefinitions/Source/Private/Material/AssetDefinition_Material.h
--- a/Engine/Plugins/Editor/EngineAssetDefinitions/Source/Private/Material/AssetDefinition_Material.h
+++ b/Engine/Plugins/Editor/EngineAssetDefinitions/Source/Private/Material/AssetDefinition_Material.h
@@ -2,12 +2,13 @@
 
 #pragma once
 
-#include "CoreMinimal.h"
+#include "Material/AssetDefinition_MaterialInterface.h"
 #include "Materials/Material.h"
-#include "AssetDefinition_MaterialInterface.h"
 
 #include "AssetDefinition_Material.generated.h"
 
+class UAssetDefinition_MaterialInterface;
+
 UCLASS()
 class UAssetDefinition_Material : public UAssetDefinition_MaterialInterface
 {
diff --git a/Engine/Plugins/Editor/EngineAssetDefinitions/Source/Private/Material/AssetDefinition_MaterialFunction.cpp b/Engine/Plugins/Editor/EngineAssetDefinitions/Source/Private/Material/AssetDefinition_MaterialFunction.cpp
--- a/Engine/Plugins/Editor/EngineAssetDefinitions/Source/Private/Material/AssetDefinition_MaterialFunction.cpp
+++ b/Engine/Plugins/Editor/EngineAssetDefinitions/Source/Private/Material/AssetDefinition_MaterialFunction.cpp
@@ -2,22 +2,17 @@
 
 #include "AssetDefinition_MaterialFunction.h"
 
+#include "AssetRegistry/AssetIdentifier.h"
 #include "ContentBrowserMenuContexts.h"
 #include "Factories/MaterialFunctionInstanceFactory.h"
+#include "Misc/DelayedAutoRegister.h"
 #include "ToolMenus.h"
-#include "UObject/UObjectHash.h"
-#include "UObject/UObjectIterator.h"
-#include "Styling/AppStyle.h"
-#include "Materials/Material.h"
-#include "Materials/MaterialFunctionInstance.h"
 #include "ThumbnailRendering/SceneThumbnailInfoWithPrimitive.h"
 #include "IAssetTools.h"
 #include "MaterialEditorModule.h"
-#include "ContentBrowserModule.h"
-#include "IContentBrowserSingleton.h"
-#include "AssetRegistry/ARFilter.h"
-#include "AssetRegistry/AssetData.h"
 #include "Misc/PackageName.h"
+#include "ToolMenu.h"
+#include "ToolMenuSection.h"
 
 #define LOCTEXT_NAMESPACE "AssetTypeActions"
 
diff --git a/Engine/Plugins/Editor/EngineAssetDefinitions/Source/Private/Material/AssetDefinition_MaterialFunction.h b/Engine/Plugins/Editor/EngineAssetDefinitions/Source/Private/Material/AssetDefinition_MaterialFunction.h
--- a/Engine/Plugins/Editor/EngineAssetDefinitions/Source/Private/Material/AssetDefinition_MaterialFunction.h
+++ b/Engine/Plugins/Editor/EngineAssetDefinitions/Source/Private/Material/AssetDefinition_MaterialFunction.h
@@ -2,8 +2,6 @@
 
 #pragma once
 
-#include "CoreMinimal.h"
-#include "Materials/MaterialFunction.h"
 #include "Materials/MaterialFunctionMaterialLayer.h"
 #include "Materials/MaterialFunctionMaterialLayerBlend.h"
 #include "AssetDefinitionDefault.h"
diff --git a/Engine/Plugins/Editor/EngineAssetDefinitions/Source/Private/Material/AssetDefinition_MaterialFunctionInstance.cpp b/Engine/Plugins/Editor/EngineAssetDefinitions/Source/Private/Material/AssetDefinition_MaterialFunctionInstance.cpp
--- a/Engine/Plugins/Editor/EngineAssetDefinitions/Source/Private/Material/AssetDefinition_MaterialFunctionInstance.cpp
+++ b/Engine/Plugins/Editor/EngineAssetDefinitions/Source/Private/Material/AssetDefinition_MaterialFunctionInstance.cpp
@@ -3,16 +3,13 @@
 #include "AssetDefinition_MaterialFunctionInstance.h"
 
 #include "ContentBrowserMenuContexts.h"
+#include "Misc/DelayedAutoRegister.h"
 #include "ToolMenus.h"
-#include "UObject/UObjectHash.h"
-#include "UObject/UObjectIterator.h"
-#include "Styling/AppStyle.h"
 #include "ThumbnailRendering/SceneThumbnailInfoWithPrimitive.h"
 #include "IAssetTools.h"
 #include "MaterialEditorModule.h"
-#include "Factories/MaterialFunctionInstanceFactory.h"
-#include "ContentBrowserModule.h"
-#include "IContentBrowserSingleton.h"
+#include "ToolMenu.h"
+#include "ToolMenuSection.h"
 
 #define LOCTEXT_NAMESPACE "AssetTypeActions"
 
diff --git a/Engine/Plugins/Editor/EngineAssetDefinitions/Source/Private/Material/AssetDefinition_MaterialFunctionInstance.h b/Engine/Plugins/Editor/EngineAssetDefinitions/Source/Private/Material/AssetDefinition_MaterialFunctionInstance.h
--- a/Engine/Plugins/Editor/EngineAssetDefinitions/Source/Private/Material/AssetDefinition_MaterialFunctionInstance.h
+++ b/Engine/Plugins/Editor/EngineAssetDefinitions/Source/Private/Material/AssetDefinition_MaterialFunctionInstance.h
@@ -2,12 +2,13 @@
 
 #pragma once
 
-#include "CoreMinimal.h"
-#include "AssetDefinition_MaterialFunction.h"
-#include "Materials/MaterialFunctionInstance.h"
 
+#include "Material/AssetDefinition_MaterialFunction.h"
 #include "AssetDefinition_MaterialFunctionInstance.generated.h"
 
+enum class EAssetCommandResult : uint8;
+struct FAssetOpenArgs;
+
 UCLASS()
 class UAssetDefinition_MaterialFunctionInstance : public UAssetDefinition_MaterialFunction
 {
diff --git a/Engine/Plugins/Editor/EngineAssetDefinitions/Source/Private/Material/AssetDefinition_MaterialInstanceConstant.cpp b/Engine/Plugins/Editor/EngineAssetDefinitions/Source/Private/Material/AssetDefinition_MaterialInstanceConstant.cpp
--- a/Engine/Plugins/Editor/EngineAssetDefinitions/Source/Private/Material/AssetDefinition_MaterialInstanceConstant.cpp
+++ b/Engine/Plugins/Editor/EngineAssetDefinitions/Source/Private/Material/AssetDefinition_MaterialInstanceConstant.cpp
@@ -3,12 +3,12 @@
 #include "AssetDefinition_MaterialInstanceConstant.h"
 
 #include "ContentBrowserMenuContexts.h"
-#include "ContentBrowserModule.h"
+#include "Misc/DelayedAutoRegister.h"
 #include "ToolMenus.h"
-#include "Styling/AppStyle.h"
 #include "IAssetTools.h"
-#include "IContentBrowserSingleton.h"
 #include "MaterialEditorModule.h"
+#include "ToolMenu.h"
+#include "ToolMenuSection.h"
 
 #define LOCTEXT_NAMESPACE "AssetTypeActions"
 
diff --git a/Engine/Plugins/Editor/EngineAssetDefinitions/Source/Private/Material/AssetDefinition_MaterialInstanceConstant.h b/Engine/Plugins/Editor/EngineAssetDefinitions/Source/Private/Material/AssetDefinition_MaterialInstanceConstant.h
--- a/Engine/Plugins/Editor/EngineAssetDefinitions/Source/Private/Material/AssetDefinition_MaterialInstanceConstant.h
+++ b/Engine/Plugins/Editor/EngineAssetDefinitions/Source/Private/Material/AssetDefinition_MaterialInstanceConstant.h
@@ -2,12 +2,16 @@
 
 #pragma once
 
-#include "CoreMinimal.h"
+#include "Material/AssetDefinition_MaterialInterface.h"
 #include "Materials/MaterialInstanceConstant.h"
-#include "AssetDefinition_MaterialInterface.h"
 
 #include "AssetDefinition_MaterialInstanceConstant.generated.h"
 
+class UAssetDefinition_MaterialInterface;
+
+enum class EAssetCommandResult : uint8;
+struct FAssetOpenArgs;
+
 UCLASS()
 class UAssetDefinition_MaterialInstanceConstant : public UAssetDefinition_MaterialInterface
 {
diff --git a/Engine/Plugins/Editor/EngineAssetDefinitions/Source/Private/Material/AssetDefinition_MaterialInstanceDynamic.h b/Engine/Plugins/Editor/EngineAssetDefinitions/Source/Private/Material/AssetDefinition_MaterialInstanceDynamic.h
--- a/Engine/Plugins/Editor/EngineAssetDefinitions/Source/Private/Material/AssetDefinition_MaterialInstanceDynamic.h
+++ b/Engine/Plugins/Editor/EngineAssetDefinitions/Source/Private/Material/AssetDefinition_MaterialInstanceDynamic.h
@@ -2,12 +2,16 @@
 
 #pragma once
 
-#include "CoreMinimal.h"
+#include "Material/AssetDefinition_MaterialInterface.h"
 #include "Materials/MaterialInstanceDynamic.h"
-#include "AssetDefinition_MaterialInterface.h"
 
 #include "AssetDefinition_MaterialInstanceDynamic.generated.h"
 
+class UAssetDefinition_MaterialInterface;
+
+enum class EAssetCommandResult : uint8;
+struct FAssetOpenArgs;
+
 UCLASS()
 class UAssetDefinition_MaterialInstanceDynamic : public UAssetDefinition_MaterialInterface
 {
diff --git a/Engine/Plugins/Editor/EngineAssetDefinitions/Source/Private/Material/AssetDefinition_MaterialInterface.cpp b/Engine/Plugins/Editor/EngineAssetDefinitions/Source/Private/Material/AssetDefinition_MaterialInterface.cpp
--- a/Engine/Plugins/Editor/EngineAssetDefinitions/Source/Private/Material/AssetDefinition_MaterialInterface.cpp
+++ b/Engine/Plugins/Editor/EngineAssetDefinitions/Source/Private/Material/AssetDefinition_MaterialInterface.cpp
@@ -4,14 +4,15 @@
 
 #include "ContentBrowserMenuContexts.h"
 #include "Materials/Material.h"
+#include "Misc/DelayedAutoRegister.h"
 #include "ToolMenus.h"
-#include "Styling/AppStyle.h"
 #include "Factories/MaterialInstanceConstantFactoryNew.h"
 #include "ThumbnailRendering/SceneThumbnailInfoWithPrimitive.h"
 #include "Materials/MaterialInstanceConstant.h"
 #include "IAssetTools.h"
-#include "IContentBrowserSingleton.h"
-#include "ContentBrowserModule.h"
+#include "ThumbnailRendering/ThumbnailManager.h"
+#include "ToolMenu.h"
+#include "ToolMenuSection.h"
 
 #define LOCTEXT_NAMESPACE "AssetTypeActions"
 
diff --git a/Engine/Plugins/Editor/GLTFImporter/Source/GLTFCore/Private/GLTF/ConversionUtilities.h b/Engine/Plugins/Editor/GLTFImporter/Source/GLTFCore/Private/GLTF/ConversionUtilities.h
--- a/Engine/Plugins/Editor/GLTFImporter/Source/GLTFCore/Private/GLTF/ConversionUtilities.h
+++ b/Engine/Plugins/Editor/GLTFImporter/Source/GLTFCore/Private/GLTF/ConversionUtilities.h
@@ -2,13 +2,12 @@
 
 #pragma once
 
-#include "CoreMinimal.h"
 
-#include "GLTFAccessor.h"
 #include "GLTFAnimation.h"
 #include "GLTFMaterial.h"
 #include "GLTFMesh.h"
 #include "GLTFTexture.h"
+#include "Math/Quat.h"
 
 namespace GLTF
 {
diff --git a/Engine/Plugins/Editor/GLTFImporter/Source/GLTFCore/Private/GLTF/ExtensionsHandler.cpp b/Engine/Plugins/Editor/GLTFImporter/Source/GLTFCore/Private/GLTF/ExtensionsHandler.cpp
--- a/Engine/Plugins/Editor/GLTFImporter/Source/GLTFCore/Private/GLTF/ExtensionsHandler.cpp
+++ b/Engine/Plugins/Editor/GLTFImporter/Source/GLTFCore/Private/GLTF/ExtensionsHandler.cpp
@@ -2,8 +2,9 @@
 
 #include "ExtensionsHandler.h"
 
+#include "GLTF/JsonUtilities.h"
 #include "GLTFAsset.h"
-#include "JsonUtilities.h"
+#include "GLTFNode.h"
 #include "MaterialUtilities.h"
 
 namespace GLTF
diff --git a/Engine/Plugins/Editor/GLTFImporter/Source/GLTFCore/Private/GLTF/GLTFAccessor.cpp b/Engine/Plugins/Editor/GLTFImporter/Source/GLTFCore/Private/GLTF/GLTFAccessor.cpp
--- a/Engine/Plugins/Editor/GLTFImporter/Source/GLTFCore/Private/GLTF/GLTFAccessor.cpp
+++ b/Engine/Plugins/Editor/GLTFImporter/Source/GLTFCore/Private/GLTF/GLTFAccessor.cpp
@@ -3,7 +3,6 @@
 #include "GLTFAccessor.h"
 
 #include "ConversionUtilities.h"
-#include "Misc/SecureHash.h"
 
 namespace GLTF
 {
diff --git a/Engine/Plugins/Editor/GLTFImporter/Source/GLTFCore/Private/GLTF/GLTFBinaryReader.cpp b/Engine/Plugins/Editor/GLTFImporter/Source/GLTFCore/Private/GLTF/GLTFBinaryReader.cpp
--- a/Engine/Plugins/Editor/GLTFImporter/Source/GLTFCore/Private/GLTF/GLTFBinaryReader.cpp
+++ b/Engine/Plugins/Editor/GLTFImporter/Source/GLTFCore/Private/GLTF/GLTFBinaryReader.cpp
@@ -2,9 +2,7 @@
 
 #include "GLTFBinaryReader.h"
 
-#include "Misc/FileHelper.h"
-#include "Serialization/ArrayReader.h"
-#include "Serialization/JsonReader.h"
+#include "CoreGlobals.h"
 
 namespace GLTF
 {
diff --git a/Engine/Plugins/Editor/GLTFImporter/Source/GLTFCore/Private/GLTF/GLTFBinaryReader.h b/Engine/Plugins/Editor/GLTFImporter/Source/GLTFCore/Private/GLTF/GLTFBinaryReader.h
--- a/Engine/Plugins/Editor/GLTFImporter/Source/GLTFCore/Private/GLTF/GLTFBinaryReader.h
+++ b/Engine/Plugins/Editor/GLTFImporter/Source/GLTFCore/Private/GLTF/GLTFBinaryReader.h
@@ -4,8 +4,8 @@
 
 #include "GLTFLogger.h"
 
-#include "Dom/JsonObject.h"
-#include "Templates/Tuple.h"
+
+struct FScriptContainerElement;
 
 namespace GLTF
 {
diff --git a/Engine/Plugins/Editor/GLTFImporter/Source/GLTFCore/Private/GLTF/GLTFMesh.cpp b/Engine/Plugins/Editor/GLTFImporter/Source/GLTFCore/Private/GLTF/GLTFMesh.cpp
--- a/Engine/Plugins/Editor/GLTFImporter/Source/GLTFCore/Private/GLTF/GLTFMesh.cpp
+++ b/Engine/Plugins/Editor/GLTFImporter/Source/GLTFCore/Private/GLTF/GLTFMesh.cpp
@@ -3,7 +3,6 @@
 #include "GLTFMesh.h"
 
 #include "ConversionUtilities.h"
-#include "Misc/SecureHash.h"
 
 namespace GLTF
 {
diff --git a/Engine/Plugins/Editor/GLTFImporter/Source/GLTFCore/Private/GLTF/GLTFReader.cpp b/Engine/Plugins/Editor/GLTFImporter/Source/GLTFCore/Private/GLTF/GLTFReader.cpp
--- a/Engine/Plugins/Editor/GLTFImporter/Source/GLTFCore/Private/GLTF/GLTFReader.cpp
+++ b/Engine/Plugins/Editor/GLTFImporter/Source/GLTFCore/Private/GLTF/GLTFReader.cpp
@@ -2,19 +2,19 @@
 
 #include "GLTFReader.h"
 
-#include "ConversionUtilities.h"
 #include "ExtensionsHandler.h"
+#include "GLTF/ConversionUtilities.h"
 #include "GLTFAsset.h"
+#include "GLTF/JsonUtilities.h"
 #include "GLTFBinaryReader.h"
-#include "JsonUtilities.h"
+#include "GLTFNode.h"
 #include "MaterialUtilities.h"
 
-#include "Dom/JsonObject.h"
+#include "HAL/PlatformFile.h"
 #include "HAL/PlatformFileManager.h"
 #include "Misc/Base64.h"
 #include "Misc/FileHelper.h"
 #include "Misc/Paths.h"
-#include "Serialization/JsonReader.h"
 #include "Serialization/JsonSerializer.h"
 
 namespace GLTF
diff --git a/Engine/Plugins/Editor/GLTFImporter/Source/GLTFCore/Private/GLTF/MaterialUtilities.h b/Engine/Plugins/Editor/GLTFImporter/Source/GLTFCore/Private/GLTF/MaterialUtilities.h
--- a/Engine/Plugins/Editor/GLTFImporter/Source/GLTFCore/Private/GLTF/MaterialUtilities.h
+++ b/Engine/Plugins/Editor/GLTFImporter/Source/GLTFCore/Private/GLTF/MaterialUtilities.h
@@ -3,11 +3,10 @@
 #pragma once
 
 #include "GLTFLogger.h"
-#include "GLTFMaterial.h"
 #include "JsonUtilities.h"
 
-#include "CoreMinimal.h"
-#include "Dom/JsonObject.h"
+
+namespace GLTF { struct FTexture; }
 
 namespace GLTF
 {
diff --git a/Engine/Plugins/Editor/GLTFImporter/Source/GLTFCore/Private/GLTFCoreModule.cpp b/Engine/Plugins/Editor/GLTFImporter/Source/GLTFCore/Private/GLTFCoreModule.cpp
--- a/Engine/Plugins/Editor/GLTFImporter/Source/GLTFCore/Private/GLTFCoreModule.cpp
+++ b/Engine/Plugins/Editor/GLTFImporter/Source/GLTFCore/Private/GLTFCoreModule.cpp
@@ -2,8 +2,6 @@
 
 #include "GLTFCoreModule.h"
 
-#include "CoreMinimal.h"
-#include "Modules/ModuleManager.h"
 
 /**
  * glTF Core module implementation (private)
@@ -17,4 +15,4 @@ public:
 	virtual void ShutdownModule() override {}
 };
 
-IMPLEMENT_MODULE(FGLTFCoreModule, GLTFCore);
\ No newline at end of file
+IMPLEMENT_MODULE(FGLTFCoreModule, GLTFCore);
diff --git a/Engine/Plugins/Editor/GLTFImporter/Source/GLTFCore/Private/GLTFMapFactory.cpp b/Engine/Plugins/Editor/GLTFImporter/Source/GLTFCore/Private/GLTFMapFactory.cpp
--- a/Engine/Plugins/Editor/GLTFImporter/Source/GLTFCore/Private/GLTFMapFactory.cpp
+++ b/Engine/Plugins/Editor/GLTFImporter/Source/GLTFCore/Private/GLTFMapFactory.cpp
@@ -5,7 +5,7 @@
 #include "GLTFMaterial.h"
 #include "GLTFMaterialExpressions.h"
 
-#include "Containers/UnrealString.h"
+#include "GLTFMaterialFactory.h"
 
 namespace GLTF
 {
diff --git a/Engine/Plugins/Editor/GLTFImporter/Source/GLTFCore/Private/GLTFMaterialExpressions.cpp b/Engine/Plugins/Editor/GLTFImporter/Source/GLTFCore/Private/GLTFMaterialExpressions.cpp
--- a/Engine/Plugins/Editor/GLTFImporter/Source/GLTFCore/Private/GLTFMaterialExpressions.cpp
+++ b/Engine/Plugins/Editor/GLTFImporter/Source/GLTFCore/Private/GLTFMaterialExpressions.cpp
@@ -2,7 +2,6 @@
 
 #include "GLTFMaterialExpressions.h"
 
-#include "Engine/Texture.h"
 
 namespace GLTF
 {
diff --git a/Engine/Plugins/Editor/GLTFImporter/Source/GLTFCore/Private/GLTFMeshFactory.cpp b/Engine/Plugins/Editor/GLTFImporter/Source/GLTFCore/Private/GLTFMeshFactory.cpp
--- a/Engine/Plugins/Editor/GLTFImporter/Source/GLTFCore/Private/GLTFMeshFactory.cpp
+++ b/Engine/Plugins/Editor/GLTFImporter/Source/GLTFCore/Private/GLTFMeshFactory.cpp
@@ -2,12 +2,9 @@
 
 #include "GLTFMeshFactory.h"
 
-#include "GLTFLogger.h"
-#include "GLTFAsset.h"
 
-#include "MeshDescription.h"
+#include "GLTFMesh.h"
 #include "StaticMeshAttributes.h"
-#include "RenderUtils.h"
 
 namespace GLTF
 {
diff --git a/Engine/Plugins/Editor/GLTFImporter/Source/GLTFImporter/Private/GLTFImportFactory.cpp b/Engine/Plugins/Editor/GLTFImporter/Source/GLTFImporter/Private/GLTFImportFactory.cpp
--- a/Engine/Plugins/Editor/GLTFImporter/Source/GLTFImporter/Private/GLTFImportFactory.cpp
+++ b/Engine/Plugins/Editor/GLTFImporter/Source/GLTFImporter/Private/GLTFImportFactory.cpp
@@ -2,15 +2,15 @@
 
 #include "GLTFImportFactory.h"
 
+#include "Framework/Application/SlateApplication.h"
 #include "GLTFImportOptions.h"
 #include "GLTFImporterContext.h"
 #include "GLTFImporterModule.h"
+#include "GLTFMesh.h"
 #include "UI/GLTFOptionsWindow.h"
 
 #include "Engine/StaticMesh.h"
 #include "Misc/FeedbackContext.h"
-#include "Misc/FileHelper.h"
-#include "Misc/Paths.h"
 
 #include "AssetRegistry/AssetRegistryModule.h"
 #include "AssetImportTask.h"
@@ -18,12 +18,11 @@
 #include "Engine/StaticMesh.h"
 #include "IMessageLogListing.h"
 #include "Interfaces/IMainFrameModule.h"
-#include "Logging/LogMacros.h"
-#include "Logging/TokenizedMessage.h"
 #include "Materials/Material.h"
 #include "MessageLogModule.h"
 #include "PackageTools.h"
-#include "UObject/StrongObjectPtr.h"
+#include "Subsystems/ImportSubsystem.h"
+#include "Widgets/SWindow.h"
 
 #include UE_INLINE_GENERATED_CPP_BY_NAME(GLTFImportFactory)
 
diff --git a/Engine/Plugins/Editor/GLTFImporter/Source/GLTFImporter/Private/GLTFImportOptions.cpp b/Engine/Plugins/Editor/GLTFImporter/Source/GLTFImporter/Private/GLTFImportOptions.cpp
--- a/Engine/Plugins/Editor/GLTFImporter/Source/GLTFImporter/Private/GLTFImportOptions.cpp
+++ b/Engine/Plugins/Editor/GLTFImporter/Source/GLTFImporter/Private/GLTFImportOptions.cpp
@@ -2,7 +2,6 @@
 
 #include "GLTFImportOptions.h"
 
-#include "CoreTypes.h"
 
 #include UE_INLINE_GENERATED_CPP_BY_NAME(GLTFImportOptions)
 
diff --git a/Engine/Plugins/Editor/GLTFImporter/Source/GLTFImporter/Private/GLTFImporterContext.cpp b/Engine/Plugins/Editor/GLTFImporter/Source/GLTFImporter/Private/GLTFImporterContext.cpp
--- a/Engine/Plugins/Editor/GLTFImporter/Source/GLTFImporter/Private/GLTFImporterContext.cpp
+++ b/Engine/Plugins/Editor/GLTFImporter/Source/GLTFImporter/Private/GLTFImporterContext.cpp
@@ -1,12 +1,16 @@
 // Copyright Epic Games, Inc. All Rights Reserved.
 
 #include "GLTFImporterContext.h"
+#include "GLTFAnimation.h"
 #include "GLTFMaterialElement.h"
-#include "GLTFMaterialExpressions.h"
+#include "GLTFMaterial.h"
 #include "GLTFTextureFactory.h"
 
+#include "GLTFMesh.h"
 #include "Materials/Material.h"
+#include "GLTFNode.h"
 #include "Misc/Paths.h"
+#include "GLTFTexture.h"
 #include "PackageTools.h"
 #include "UObject/Package.h"
 
diff --git a/Engine/Plugins/Editor/GLTFImporter/Source/GLTFImporter/Private/GLTFImporterContext.h b/Engine/Plugins/Editor/GLTFImporter/Source/GLTFImporter/Private/GLTFImporterContext.h
--- a/Engine/Plugins/Editor/GLTFImporter/Source/GLTFImporter/Private/GLTFImporterContext.h
+++ b/Engine/Plugins/Editor/GLTFImporter/Source/GLTFImporter/Private/GLTFImporterContext.h
@@ -3,7 +3,6 @@
 #pragma once
 
 #include "GLTFAsset.h"
-#include "GLTFLogger.h"
 #include "GLTFMaterialFactory.h"
 #include "GLTFReader.h"
 #include "GLTFStaticMeshFactory.h"
diff --git a/Engine/Plugins/Editor/GLTFImporter/Source/GLTFImporter/Private/GLTFImporterModule.cpp b/Engine/Plugins/Editor/GLTFImporter/Source/GLTFImporter/Private/GLTFImporterModule.cpp
--- a/Engine/Plugins/Editor/GLTFImporter/Source/GLTFImporter/Private/GLTFImporterModule.cpp
+++ b/Engine/Plugins/Editor/GLTFImporter/Source/GLTFImporter/Private/GLTFImporterModule.cpp
@@ -2,10 +2,13 @@
 
 #include "GLTFImporterModule.h"
 
+#include "GLTFAnimation.h"
 #include "GLTFImporterContext.h"
 
-#include "CoreMinimal.h"
-#include "Modules/ModuleManager.h"
+#include "GLTFMaterial.h"
+#include "GLTFMesh.h"
+#include "GLTFNode.h"
+#include "GLTFTexture.h"
 
 /**
  * glTF Importer module implementation (private)
diff --git a/Engine/Plugins/Editor/GLTFImporter/Source/GLTFImporter/Private/GLTFMaterialElement.cpp b/Engine/Plugins/Editor/GLTFImporter/Source/GLTFImporter/Private/GLTFMaterialElement.cpp
--- a/Engine/Plugins/Editor/GLTFImporter/Source/GLTFImporter/Private/GLTFMaterialElement.cpp
+++ b/Engine/Plugins/Editor/GLTFImporter/Source/GLTFImporter/Private/GLTFMaterialElement.cpp
@@ -4,14 +4,12 @@
 
 #include "AssetRegistry/AssetRegistryModule.h"
 #include "MaterialEditingLibrary.h"
-#include "MaterialEditorUtilities.h"
 #include "Materials/Material.h"
 #include "Materials/MaterialExpressionConstant.h"
 #include "Materials/MaterialExpressionConstant3Vector.h"
 #include "Materials/MaterialExpressionMaterialFunctionCall.h"
 #include "Materials/MaterialExpressionScalarParameter.h"
 #include "Materials/MaterialExpressionTextureCoordinate.h"
-#include "Materials/MaterialExpressionTextureObjectParameter.h"
 #include "Materials/MaterialExpressionTextureSampleParameter2D.h"
 #include "Materials/MaterialExpressionVectorParameter.h"
 #include "Materials/MaterialExpressionThinTranslucentMaterialOutput.h"
diff --git a/Engine/Plugins/Editor/GLTFImporter/Source/GLTFImporter/Private/GLTFMaterialElement.h b/Engine/Plugins/Editor/GLTFImporter/Source/GLTFImporter/Private/GLTFMaterialElement.h
--- a/Engine/Plugins/Editor/GLTFImporter/Source/GLTFImporter/Private/GLTFMaterialElement.h
+++ b/Engine/Plugins/Editor/GLTFImporter/Source/GLTFImporter/Private/GLTFMaterialElement.h
@@ -3,6 +3,7 @@
 #pragma once
 
 #include "GLTFMaterialExpressions.h"
+#include "UObject/StrongObjectPtr.h"
 
 class UMaterialExpression;
 class UMaterial;
diff --git a/Engine/Plugins/Editor/GLTFImporter/Source/GLTFImporter/Private/GLTFStaticMeshFactory.cpp b/Engine/Plugins/Editor/GLTFImporter/Source/GLTFImporter/Private/GLTFStaticMeshFactory.cpp
--- a/Engine/Plugins/Editor/GLTFImporter/Source/GLTFImporter/Private/GLTFStaticMeshFactory.cpp
+++ b/Engine/Plugins/Editor/GLTFImporter/Source/GLTFImporter/Private/GLTFStaticMeshFactory.cpp
@@ -2,13 +2,15 @@
 
 #include "GLTFStaticMeshFactory.h"
 
+#include "Engine/StaticMeshSourceData.h"
 #include "GLTFAsset.h"
+#include "GLTFMesh.h"
 #include "GLTFMeshFactory.h"
 
 #include "AssetRegistry/AssetRegistryModule.h"
 #include "Engine/StaticMesh.h"
+#include "MeshDescription.h"
 #include "Misc/Paths.h"
-#include "StaticMeshAttributes.h"
 #include "PackageTools.h"
 #include "UObject/Package.h"
 
diff --git a/Engine/Plugins/Editor/GLTFImporter/Source/GLTFImporter/Private/GLTFTextureFactory.cpp b/Engine/Plugins/Editor/GLTFImporter/Source/GLTFImporter/Private/GLTFTextureFactory.cpp
--- a/Engine/Plugins/Editor/GLTFImporter/Source/GLTFImporter/Private/GLTFTextureFactory.cpp
+++ b/Engine/Plugins/Editor/GLTFImporter/Source/GLTFImporter/Private/GLTFTextureFactory.cpp
@@ -2,7 +2,6 @@
 
 #include "GLTFTextureFactory.h"
 
-#include "GLTFMaterial.h"
 #include "GLTFMaterialElement.h"
 #include "GLTFTexture.h"
 
diff --git a/Engine/Plugins/Editor/GLTFImporter/Source/GLTFImporter/Private/GLTFTextureFactory.h b/Engine/Plugins/Editor/GLTFImporter/Source/GLTFImporter/Private/GLTFTextureFactory.h
--- a/Engine/Plugins/Editor/GLTFImporter/Source/GLTFImporter/Private/GLTFTextureFactory.h
+++ b/Engine/Plugins/Editor/GLTFImporter/Source/GLTFImporter/Private/GLTFTextureFactory.h
@@ -2,11 +2,8 @@
 
 #pragma once
 
-#include "CoreMinimal.h"
-#include "UObject/ObjectMacros.h"
 #include "UObject/StrongObjectPtr.h"
 
-#include "GLTFLogger.h"
 #include "GLTFMaterialFactory.h"
 
 class UTexture;
diff --git a/Engine/Plugins/Editor/GLTFImporter/Source/GLTFImporter/Private/UI/GLTFOptionsWindow.cpp b/Engine/Plugins/Editor/GLTFImporter/Source/GLTFImporter/Private/UI/GLTFOptionsWindow.cpp
--- a/Engine/Plugins/Editor/GLTFImporter/Source/GLTFImporter/Private/UI/GLTFOptionsWindow.cpp
+++ b/Engine/Plugins/Editor/GLTFImporter/Source/GLTFImporter/Private/UI/GLTFOptionsWindow.cpp
@@ -2,7 +2,7 @@
 
 #include "GLTFOptionsWindow.h"
 
-#include "Styling/AppStyle.h"
+#include "DetailsViewArgs.h"
 #include "HAL/PlatformProcess.h"
 #include "IDetailsView.h"
 #include "Modules/ModuleManager.h"
@@ -10,6 +10,7 @@
 #include "Widgets/Input/SButton.h"
 #include "Widgets/Layout/SBox.h"
 #include "Widgets/Layout/SUniformGridPanel.h"
+#include "Widgets/SWindow.h"
 #include "Widgets/Text/SInlineEditableTextBlock.h"
 
 #define LOCTEXT_NAMESPACE "GLTFOptionsWindow"
diff --git a/Engine/Plugins/Editor/GLTFImporter/Source/GLTFImporter/Private/UI/GLTFOptionsWindow.h b/Engine/Plugins/Editor/GLTFImporter/Source/GLTFImporter/Private/UI/GLTFOptionsWindow.h
--- a/Engine/Plugins/Editor/GLTFImporter/Source/GLTFImporter/Private/UI/GLTFOptionsWindow.h
+++ b/Engine/Plugins/Editor/GLTFImporter/Source/GLTFImporter/Private/UI/GLTFOptionsWindow.h
@@ -2,7 +2,6 @@
 
 #pragma once
 
-#include "Framework/Application/SlateApplication.h"
 #include "Widgets/SCompoundWidget.h"
 
 class SGLTFOptionsWindow : public SCompoundWidget
diff --git a/Engine/Plugins/Editor/GameplayTagsEditor/Source/GameplayTagsEditor/Private/AssetTypeActions_GameplayTagAssetBase.cpp b/Engine/Plugins/Editor/GameplayTagsEditor/Source/GameplayTagsEditor/Private/AssetTypeActions_GameplayTagAssetBase.cpp
--- a/Engine/Plugins/Editor/GameplayTagsEditor/Source/GameplayTagsEditor/Private/AssetTypeActions_GameplayTagAssetBase.cpp
+++ b/Engine/Plugins/Editor/GameplayTagsEditor/Source/GameplayTagsEditor/Private/AssetTypeActions_GameplayTagAssetBase.cpp
@@ -1,13 +1,12 @@
 // Copyright Epic Games, Inc. All Rights Reserved.
 
 #include "AssetTypeActions_GameplayTagAssetBase.h"
-#include "ToolMenus.h"
-#include "UObject/UnrealType.h"
 #include "Framework/Application/SlateApplication.h"
 
 #include "GameplayTagContainer.h"
 #include "SGameplayTagWidget.h"
 #include "Interfaces/IMainFrameModule.h"
+#include "ToolMenuSection.h"
 
 #define LOCTEXT_NAMESPACE "AssetTypeActions"
 
diff --git a/Engine/Plugins/Editor/GameplayTagsEditor/Source/GameplayTagsEditor/Private/GameplayTagContainerCustomization.cpp b/Engine/Plugins/Editor/GameplayTagsEditor/Source/GameplayTagsEditor/Private/GameplayTagContainerCustomization.cpp
--- a/Engine/Plugins/Editor/GameplayTagsEditor/Source/GameplayTagsEditor/Private/GameplayTagContainerCustomization.cpp
+++ b/Engine/Plugins/Editor/GameplayTagsEditor/Source/GameplayTagsEditor/Private/GameplayTagContainerCustomization.cpp
@@ -1,19 +1,20 @@
 // Copyright Epic Games, Inc. All Rights Reserved.
 
 #include "GameplayTagContainerCustomization.h"
+#include "AssetRegistry/AssetIdentifier.h"
 #include "Widgets/Input/SComboButton.h"
 
-#include "Widgets/Input/SButton.h"
 
 
 #include "Editor.h"
-#include "PropertyHandle.h"
 #include "DetailWidgetRow.h"
-#include "Framework/Application/MenuStack.h"
 #include "Framework/Application/SlateApplication.h"
+#include "Framework/Views/TableViewMetadata.h"
 #include "GameplayTagsEditorModule.h"
+#include "GameplayTagsManager.h"
 #include "HAL/PlatformApplicationMisc.h"
 #include "Layout/WidgetPath.h"
+#include "SGameplayTagWidget.h"
 #include "ScopedTransaction.h"
 #include "Widgets/Input/SHyperlink.h"
 #include "EditorFontGlyphs.h"
diff --git a/Engine/Plugins/Editor/GameplayTagsEditor/Source/GameplayTagsEditor/Private/GameplayTagContainerCustomization.h b/Engine/Plugins/Editor/GameplayTagsEditor/Source/GameplayTagsEditor/Private/GameplayTagContainerCustomization.h
--- a/Engine/Plugins/Editor/GameplayTagsEditor/Source/GameplayTagsEditor/Private/GameplayTagContainerCustomization.h
+++ b/Engine/Plugins/Editor/GameplayTagsEditor/Source/GameplayTagsEditor/Private/GameplayTagContainerCustomization.h
@@ -2,15 +2,12 @@
 
 #pragma once
 
-#include "CoreMinimal.h"
-#include "Layout/Visibility.h"
-#include "Widgets/SWidget.h"
 #include "IPropertyTypeCustomization.h"
 #include "GameplayTagContainerCustomizationOptions.h"
-#include "Widgets/Views/STableViewBase.h"
-#include "Widgets/Views/STableRow.h"
-#include "SGameplayTagWidget.h"
 #include "EditorUndoClient.h"
+#include "Widgets/Views/SListView.h"
+
+struct FGameplayTag;
 
 class IPropertyHandle;
 
diff --git a/Engine/Plugins/Editor/GameplayTagsEditor/Source/GameplayTagsEditor/Private/GameplayTagCustomization.cpp b/Engine/Plugins/Editor/GameplayTagsEditor/Source/GameplayTagsEditor/Private/GameplayTagCustomization.cpp
--- a/Engine/Plugins/Editor/GameplayTagsEditor/Source/GameplayTagsEditor/Private/GameplayTagCustomization.cpp
+++ b/Engine/Plugins/Editor/GameplayTagsEditor/Source/GameplayTagsEditor/Private/GameplayTagCustomization.cpp
@@ -1,11 +1,12 @@
 // Copyright Epic Games, Inc. All Rights Reserved.
 
 #include "GameplayTagCustomization.h"
+#include "AssetRegistry/AssetIdentifier.h"
 #include "Widgets/Input/SComboButton.h"
 
 #include "Editor.h"
-#include "PropertyHandle.h"
 #include "DetailWidgetRow.h"
+#include "GameplayTagsManager.h"
 #include "IDetailChildrenBuilder.h"
 #include "GameplayTagsEditorModule.h"
 #include "Widgets/Input/SHyperlink.h"
diff --git a/Engine/Plugins/Editor/GameplayTagsEditor/Source/GameplayTagsEditor/Private/GameplayTagCustomization.h b/Engine/Plugins/Editor/GameplayTagsEditor/Source/GameplayTagsEditor/Private/GameplayTagCustomization.h
--- a/Engine/Plugins/Editor/GameplayTagsEditor/Source/GameplayTagsEditor/Private/GameplayTagCustomization.h
+++ b/Engine/Plugins/Editor/GameplayTagsEditor/Source/GameplayTagsEditor/Private/GameplayTagCustomization.h
@@ -2,13 +2,13 @@
 
 #pragma once
 
-#include "CoreMinimal.h"
-#include "Widgets/SWidget.h"
 #include "IPropertyTypeCustomization.h"
 #include "GameplayTagCustomizationOptions.h"
 #include "SGameplayTagWidget.h"
 #include "EditorUndoClient.h"
 
+struct EVisibility;
+
 class IPropertyHandle;
 
 /** Customization for the gameplay tag struct */
diff --git a/Engine/Plugins/Editor/GameplayTagsEditor/Source/GameplayTagsEditor/Private/GameplayTagPinUtilities.cpp b/Engine/Plugins/Editor/GameplayTagsEditor/Source/GameplayTagsEditor/Private/GameplayTagPinUtilities.cpp
--- a/Engine/Plugins/Editor/GameplayTagsEditor/Source/GameplayTagsEditor/Private/GameplayTagPinUtilities.cpp
+++ b/Engine/Plugins/Editor/GameplayTagsEditor/Source/GameplayTagsEditor/Private/GameplayTagPinUtilities.cpp
@@ -2,7 +2,6 @@
 
 #include "GameplayTagPinUtilities.h"
 #include "GameplayTagsManager.h"
-#include "SGraphPin.h"
 #include "K2Node_CallFunction.h"
 #include "K2Node_VariableSet.h"
 #include "K2Node_FunctionTerminator.h"
diff --git a/Engine/Plugins/Editor/GameplayTagsEditor/Source/GameplayTagsEditor/Private/GameplayTagPinUtilities.h b/Engine/Plugins/Editor/GameplayTagsEditor/Source/GameplayTagsEditor/Private/GameplayTagPinUtilities.h
--- a/Engine/Plugins/Editor/GameplayTagsEditor/Source/GameplayTagsEditor/Private/GameplayTagPinUtilities.h
+++ b/Engine/Plugins/Editor/GameplayTagsEditor/Source/GameplayTagsEditor/Private/GameplayTagPinUtilities.h
@@ -2,7 +2,8 @@
 
 #pragma once
 
-#include "CoreMinimal.h"
+
+class FString;
 
 // Forward declarations
 class UEdGraphPin;
diff --git a/Engine/Plugins/Editor/GameplayTagsEditor/Source/GameplayTagsEditor/Private/GameplayTagQueryCustomization.cpp b/Engine/Plugins/Editor/GameplayTagsEditor/Source/GameplayTagsEditor/Private/GameplayTagQueryCustomization.cpp
--- a/Engine/Plugins/Editor/GameplayTagsEditor/Source/GameplayTagsEditor/Private/GameplayTagQueryCustomization.cpp
+++ b/Engine/Plugins/Editor/GameplayTagsEditor/Source/GameplayTagsEditor/Private/GameplayTagQueryCustomization.cpp
@@ -1,21 +1,14 @@
 // Copyright Epic Games, Inc. All Rights Reserved.
 
 #include "GameplayTagQueryCustomization.h"
-#include "UObject/UnrealType.h"
-#include "Widgets/SBoxPanel.h"
-#include "Widgets/SWindow.h"
 #include "Framework/Application/SlateApplication.h"
-#include "Widgets/Layout/SBorder.h"
+#include "GameplayTagContainer.h"
 #include "Widgets/Text/STextBlock.h"
 #include "Widgets/Input/SButton.h"
 #include "Framework/Docking/TabManager.h"
 
 #include "Editor.h"
-#include "PropertyHandle.h"
 #include "DetailWidgetRow.h"
-#include "IPropertyTypeCustomization.h"
-#include "IPropertyUtilities.h"
-#include "Misc/NotifyHook.h"
 
 #define LOCTEXT_NAMESPACE "GameplayTagQueryCustomization"
 
diff --git a/Engine/Plugins/Editor/GameplayTagsEditor/Source/GameplayTagsEditor/Private/GameplayTagQueryCustomization.h b/Engine/Plugins/Editor/GameplayTagsEditor/Source/GameplayTagsEditor/Private/GameplayTagQueryCustomization.h
--- a/Engine/Plugins/Editor/GameplayTagsEditor/Source/GameplayTagsEditor/Private/GameplayTagQueryCustomization.h
+++ b/Engine/Plugins/Editor/GameplayTagsEditor/Source/GameplayTagsEditor/Private/GameplayTagQueryCustomization.h
@@ -2,13 +2,12 @@
 
 #pragma once
 
-#include "CoreMinimal.h"
-#include "Layout/Visibility.h"
-#include "Input/Reply.h"
 #include "IPropertyTypeCustomization.h"
 #include "SGameplayTagQueryWidget.h"
 #include "EditorUndoClient.h"
 
+struct EVisibility;
+
 class IPropertyHandle;
 class SWindow;
 
diff --git a/Engine/Plugins/Editor/GameplayTagsEditor/Source/GameplayTagsEditor/Private/GameplayTagSearchFilter.cpp b/Engine/Plugins/Editor/GameplayTagsEditor/Source/GameplayTagsEditor/Private/GameplayTagSearchFilter.cpp
--- a/Engine/Plugins/Editor/GameplayTagsEditor/Source/GameplayTagsEditor/Private/GameplayTagSearchFilter.cpp
+++ b/Engine/Plugins/Editor/GameplayTagsEditor/Source/GameplayTagsEditor/Private/GameplayTagSearchFilter.cpp
@@ -1,14 +1,12 @@
 // Copyright Epic Games, Inc. All Rights Reserved.
 
 #include "GameplayTagSearchFilter.h"
-#include "Framework/Commands/UIAction.h"
 #include "Engine/Blueprint.h"
-#include "UObject/UnrealType.h"
+#include "FrontendFilterBase.h"
 #include "Misc/ConfigCacheIni.h"
 #include "Widgets/SBoxPanel.h"
 #include "Framework/MultiBox/MultiBoxBuilder.h"
 
-#include "GameplayTagContainer.h"
 #include "GameplayTagsManager.h"
 #include "SGameplayTagWidget.h"
 
diff --git a/Engine/Plugins/Editor/GameplayTagsEditor/Source/GameplayTagsEditor/Private/GameplayTagSearchFilter.h b/Engine/Plugins/Editor/GameplayTagsEditor/Source/GameplayTagsEditor/Private/GameplayTagSearchFilter.h
--- a/Engine/Plugins/Editor/GameplayTagsEditor/Source/GameplayTagsEditor/Private/GameplayTagSearchFilter.h
+++ b/Engine/Plugins/Editor/GameplayTagsEditor/Source/GameplayTagsEditor/Private/GameplayTagSearchFilter.h
@@ -2,9 +2,6 @@
 
 #pragma once
 
-#include "CoreMinimal.h"
-#include "UObject/ObjectMacros.h"
-#include "FrontendFilterBase.h"
 #include "ContentBrowserFrontEndFilterExtension.h"
 #include "GameplayTagSearchFilter.generated.h"
 
diff --git a/Engine/Plugins/Editor/GameplayTagsEditor/Source/GameplayTagsEditor/Private/GameplayTagsEditorModule.cpp b/Engine/Plugins/Editor/GameplayTagsEditor/Source/GameplayTagsEditor/Private/GameplayTagsEditorModule.cpp
--- a/Engine/Plugins/Editor/GameplayTagsEditor/Source/GameplayTagsEditor/Private/GameplayTagsEditorModule.cpp
+++ b/Engine/Plugins/Editor/GameplayTagsEditor/Source/GameplayTagsEditor/Private/GameplayTagsEditorModule.cpp
@@ -2,27 +2,22 @@
 
 #include "GameplayTagsEditorModule.h"
 
-#include "Misc/Paths.h"
+#include "GameplayTagRedirectors.h"
 #include "PropertyEditorModule.h"
 #include "Factories/Factory.h"
-#include "EdGraphUtilities.h"
-#include "GameplayTagContainer.h"
-#include "Engine/DataTable.h"
 #include "GameplayTagsManager.h"
-#include "AssetRegistry/AssetData.h"
-#include "Misc/ConfigCacheIni.h"
 #include "GameplayTagsGraphPanelPinFactory.h"
 #include "GameplayTagsGraphPanelNodeFactory.h"
-#include "GameplayTagContainerCustomization.h"
 #include "GameplayTagQueryCustomization.h"
 #include "GameplayTagCustomization.h"
 #include "GameplayTagsSettings.h"
 #include "GameplayTagsSettingsCustomization.h"
 #include "GameplayTagsModule.h"
+#include "HAL/FileManager.h"
 #include "ISettingsModule.h"
 #include "ISettingsEditorModule.h"
+#include "Misc/CoreDelegates.h"
 #include "Widgets/Notifications/SNotificationList.h"
-#include "Widgets/SNullWidget.h"
 #include "Framework/Notifications/NotificationManager.h"
 #include "AssetRegistry/AssetRegistryModule.h"
 #include "Editor.h"
@@ -31,9 +26,7 @@
 #include "HAL/PlatformFile.h"
 #include "HAL/PlatformFileManager.h"
 #include "Stats/StatsMisc.h"
-#include "UObject/UObjectHash.h"
-#include "HAL/IConsoleManager.h"
-#include "Misc/FileHelper.h"
+#include "Subsystems/ImportSubsystem.h"
 #include "UObject/ObjectSaveContext.h"
 
 #define LOCTEXT_NAMESPACE "GameplayTagEditor"
diff --git a/Engine/Plugins/Editor/GameplayTagsEditor/Source/GameplayTagsEditor/Private/GameplayTagsGraphPanelNodeFactory.h b/Engine/Plugins/Editor/GameplayTagsEditor/Source/GameplayTagsEditor/Private/GameplayTagsGraphPanelNodeFactory.h
--- a/Engine/Plugins/Editor/GameplayTagsEditor/Source/GameplayTagsEditor/Private/GameplayTagsGraphPanelNodeFactory.h
+++ b/Engine/Plugins/Editor/GameplayTagsEditor/Source/GameplayTagsEditor/Private/GameplayTagsGraphPanelNodeFactory.h
@@ -2,10 +2,7 @@
 
 #pragma once
 
-#include "CoreMinimal.h"
-#include "Widgets/DeclarativeSyntaxSupport.h"
 #include "EdGraphUtilities.h"
-#include "SGraphNode.h"
 #include "SGraphNode_MultiCompareGameplayTag.h"
 #include "GameplayTagsK2Node_MultiCompareBase.h"
 
diff --git a/Engine/Plugins/Editor/GameplayTagsEditor/Source/GameplayTagsEditor/Private/GameplayTagsGraphPanelPinFactory.h b/Engine/Plugins/Editor/GameplayTagsEditor/Source/GameplayTagsEditor/Private/GameplayTagsGraphPanelPinFactory.h
--- a/Engine/Plugins/Editor/GameplayTagsEditor/Source/GameplayTagsEditor/Private/GameplayTagsGraphPanelPinFactory.h
+++ b/Engine/Plugins/Editor/GameplayTagsEditor/Source/GameplayTagsEditor/Private/GameplayTagsGraphPanelPinFactory.h
@@ -2,13 +2,9 @@
 
 #pragma once
 
-#include "CoreMinimal.h"
-#include "Widgets/DeclarativeSyntaxSupport.h"
 #include "EdGraphUtilities.h"
 #include "GameplayTagContainer.h"
 #include "EdGraphSchema_K2.h"
-#include "SGraphPin.h"
-#include "SGameplayTagGraphPin.h"
 #include "SGameplayTagContainerGraphPin.h"
 #include "SGameplayTagQueryGraphPin.h"
 
diff --git a/Engine/Plugins/Editor/GameplayTagsEditor/Source/GameplayTagsEditor/Private/GameplayTagsK2Node_LiteralGameplayTag.cpp b/Engine/Plugins/Editor/GameplayTagsEditor/Source/GameplayTagsEditor/Private/GameplayTagsK2Node_LiteralGameplayTag.cpp
--- a/Engine/Plugins/Editor/GameplayTagsEditor/Source/GameplayTagsEditor/Private/GameplayTagsK2Node_LiteralGameplayTag.cpp
+++ b/Engine/Plugins/Editor/GameplayTagsEditor/Source/GameplayTagsEditor/Private/GameplayTagsK2Node_LiteralGameplayTag.cpp
@@ -2,13 +2,8 @@
 
 #include "GameplayTagsK2Node_LiteralGameplayTag.h"
 #include "EdGraphSchema_K2.h"
-#include "K2Node_MakeArray.h"
-#include "K2Node_MakeStruct.h"
 
-#include "KismetCompiler.h"
 #include "GameplayTagContainer.h"
-#include "BlueprintNodeSpawner.h"
-#include "BlueprintActionDatabaseRegistrar.h"
 #include "BlueprintGameplayTagLibrary.h"
 
 #include UE_INLINE_GENERATED_CPP_BY_NAME(GameplayTagsK2Node_LiteralGameplayTag)
diff --git a/Engine/Plugins/Editor/GameplayTagsEditor/Source/GameplayTagsEditor/Private/GameplayTagsK2Node_MultiCompareBase.cpp b/Engine/Plugins/Editor/GameplayTagsEditor/Source/GameplayTagsEditor/Private/GameplayTagsK2Node_MultiCompareBase.cpp
--- a/Engine/Plugins/Editor/GameplayTagsEditor/Source/GameplayTagsEditor/Private/GameplayTagsK2Node_MultiCompareBase.cpp
+++ b/Engine/Plugins/Editor/GameplayTagsEditor/Source/GameplayTagsEditor/Private/GameplayTagsK2Node_MultiCompareBase.cpp
@@ -1,7 +1,6 @@
 // Copyright Epic Games, Inc. All Rights Reserved.
 
 #include "GameplayTagsK2Node_MultiCompareBase.h"
-#include "UObject/UnrealType.h"
 #include "EdGraph/EdGraph.h"
 
 #include UE_INLINE_GENERATED_CPP_BY_NAME(GameplayTagsK2Node_MultiCompareBase)
diff --git a/Engine/Plugins/Editor/GameplayTagsEditor/Source/GameplayTagsEditor/Private/GameplayTagsK2Node_MultiCompareGameplayTagAssetInterface.cpp b/Engine/Plugins/Editor/GameplayTagsEditor/Source/GameplayTagsEditor/Private/GameplayTagsK2Node_MultiCompareGameplayTagAssetInterface.cpp
--- a/Engine/Plugins/Editor/GameplayTagsEditor/Source/GameplayTagsEditor/Private/GameplayTagsK2Node_MultiCompareGameplayTagAssetInterface.cpp
+++ b/Engine/Plugins/Editor/GameplayTagsEditor/Source/GameplayTagsEditor/Private/GameplayTagsK2Node_MultiCompareGameplayTagAssetInterface.cpp
@@ -2,7 +2,6 @@
 
 #include "GameplayTagsK2Node_MultiCompareGameplayTagAssetInterface.h"
 #include "GameplayTagContainer.h"
-#include "EdGraphSchema_K2.h"
 #include "K2Node_CallFunction.h"
 
 #include "BlueprintNodeSpawner.h"
diff --git a/Engine/Plugins/Editor/GameplayTagsEditor/Source/GameplayTagsEditor/Private/GameplayTagsK2Node_MultiCompareGameplayTagAssetInterfaceSingleTags.cpp b/Engine/Plugins/Editor/GameplayTagsEditor/Source/GameplayTagsEditor/Private/GameplayTagsK2Node_MultiCompareGameplayTagAssetInterfaceSingleTags.cpp
--- a/Engine/Plugins/Editor/GameplayTagsEditor/Source/GameplayTagsEditor/Private/GameplayTagsK2Node_MultiCompareGameplayTagAssetInterfaceSingleTags.cpp
+++ b/Engine/Plugins/Editor/GameplayTagsEditor/Source/GameplayTagsEditor/Private/GameplayTagsK2Node_MultiCompareGameplayTagAssetInterfaceSingleTags.cpp
@@ -2,7 +2,6 @@
 
 #include "GameplayTagsK2Node_MultiCompareGameplayTagAssetInterfaceSingleTags.h"
 #include "GameplayTagContainer.h"
-#include "EdGraphSchema_K2.h"
 #include "K2Node_CallFunction.h"
 #include "BlueprintNodeSpawner.h"
 #include "BlueprintActionDatabaseRegistrar.h"
diff --git a/Engine/Plugins/Editor/GameplayTagsEditor/Source/GameplayTagsEditor/Private/GameplayTagsK2Node_SwitchGameplayTag.cpp b/Engine/Plugins/Editor/GameplayTagsEditor/Source/GameplayTagsEditor/Private/GameplayTagsK2Node_SwitchGameplayTag.cpp
--- a/Engine/Plugins/Editor/GameplayTagsEditor/Source/GameplayTagsEditor/Private/GameplayTagsK2Node_SwitchGameplayTag.cpp
+++ b/Engine/Plugins/Editor/GameplayTagsEditor/Source/GameplayTagsEditor/Private/GameplayTagsK2Node_SwitchGameplayTag.cpp
@@ -1,6 +1,7 @@
 // Copyright Epic Games, Inc. All Rights Reserved.
 
 #include "GameplayTagsK2Node_SwitchGameplayTag.h"
+#include "EdGraph/EdGraph.h"
 #include "EdGraphSchema_K2.h"
 #include "Engine/Blueprint.h"
 #include "BlueprintNodeSpawner.h"
diff --git a/Engine/Plugins/Editor/GameplayTagsEditor/Source/GameplayTagsEditor/Private/GameplayTagsK2Node_SwitchGameplayTagContainer.cpp b/Engine/Plugins/Editor/GameplayTagsEditor/Source/GameplayTagsEditor/Private/GameplayTagsK2Node_SwitchGameplayTagContainer.cpp
--- a/Engine/Plugins/Editor/GameplayTagsEditor/Source/GameplayTagsEditor/Private/GameplayTagsK2Node_SwitchGameplayTagContainer.cpp
+++ b/Engine/Plugins/Editor/GameplayTagsEditor/Source/GameplayTagsEditor/Private/GameplayTagsK2Node_SwitchGameplayTagContainer.cpp
@@ -1,11 +1,13 @@
 // Copyright Epic Games, Inc. All Rights Reserved.
 
 #include "GameplayTagsK2Node_SwitchGameplayTagContainer.h"
+#include "EdGraph/EdGraph.h"
 #include "EdGraphSchema_K2.h"
 #include "Engine/Blueprint.h"
 #include "BlueprintNodeSpawner.h"
 #include "BlueprintActionDatabaseRegistrar.h"
 #include "BlueprintGameplayTagLibrary.h"
+#include "GameplayTagContainer.h"
 
 #include UE_INLINE_GENERATED_CPP_BY_NAME(GameplayTagsK2Node_SwitchGameplayTagContainer)
 
diff --git a/Engine/Plugins/Editor/GameplayTagsEditor/Source/GameplayTagsEditor/Private/GameplayTagsSettingsCustomization.cpp b/Engine/Plugins/Editor/GameplayTagsEditor/Source/GameplayTagsEditor/Private/GameplayTagsSettingsCustomization.cpp
--- a/Engine/Plugins/Editor/GameplayTagsEditor/Source/GameplayTagsEditor/Private/GameplayTagsSettingsCustomization.cpp
+++ b/Engine/Plugins/Editor/GameplayTagsEditor/Source/GameplayTagsEditor/Private/GameplayTagsSettingsCustomization.cpp
@@ -3,10 +3,9 @@
 #include "GameplayTagsSettingsCustomization.h"
 #include "GameplayTagsSettings.h"
 #include "GameplayTagsModule.h"
-#include "PropertyHandle.h"
-#include "DetailCategoryBuilder.h"
 #include "DetailLayoutBuilder.h"
 #include "DetailWidgetRow.h"
+#include "SGameplayTagWidget.h"
 
 #define LOCTEXT_NAMESPACE "FGameplayTagsSettingsCustomization"
 
diff --git a/Engine/Plugins/Editor/GameplayTagsEditor/Source/GameplayTagsEditor/Private/GameplayTagsSettingsCustomization.h b/Engine/Plugins/Editor/GameplayTagsEditor/Source/GameplayTagsEditor/Private/GameplayTagsSettingsCustomization.h
--- a/Engine/Plugins/Editor/GameplayTagsEditor/Source/GameplayTagsEditor/Private/GameplayTagsSettingsCustomization.h
+++ b/Engine/Plugins/Editor/GameplayTagsEditor/Source/GameplayTagsEditor/Private/GameplayTagsSettingsCustomization.h
@@ -2,9 +2,9 @@
 
 #pragma once
 
-#include "CoreMinimal.h"
 #include "IDetailCustomization.h"
-#include "SGameplayTagWidget.h"
+
+class SGameplayTagWidget;
 
 class IDetailLayoutBuilder;
 
diff --git a/Engine/Plugins/Editor/GameplayTagsEditor/Source/GameplayTagsEditor/Private/SAddNewGameplayTagSourceWidget.cpp b/Engine/Plugins/Editor/GameplayTagsEditor/Source/GameplayTagsEditor/Private/SAddNewGameplayTagSourceWidget.cpp
--- a/Engine/Plugins/Editor/GameplayTagsEditor/Source/GameplayTagsEditor/Private/SAddNewGameplayTagSourceWidget.cpp
+++ b/Engine/Plugins/Editor/GameplayTagsEditor/Source/GameplayTagsEditor/Private/SAddNewGameplayTagSourceWidget.cpp
@@ -2,10 +2,13 @@
 
 #include "SAddNewGameplayTagSourceWidget.h"
 #include "DetailLayoutBuilder.h"
-#include "GameplayTagsSettings.h"
+#include "Framework/Views/TableViewMetadata.h"
 #include "GameplayTagsEditorModule.h"
-#include "GameplayTagsModule.h"
+#include "GameplayTagsManager.h"
 #include "Widgets/Input/SButton.h"
+#include "Misc/Paths.h"
+#include "Widgets/Input/SComboBox.h"
+#include "Widgets/Input/SEditableTextBox.h"
 
 #define LOCTEXT_NAMESPACE "AddNewGameplayTagSourceWidget"
 
diff --git a/Engine/Plugins/Editor/GameplayTagsEditor/Source/GameplayTagsEditor/Private/SAddNewGameplayTagSourceWidget.h b/Engine/Plugins/Editor/GameplayTagsEditor/Source/GameplayTagsEditor/Private/SAddNewGameplayTagSourceWidget.h
--- a/Engine/Plugins/Editor/GameplayTagsEditor/Source/GameplayTagsEditor/Private/SAddNewGameplayTagSourceWidget.h
+++ b/Engine/Plugins/Editor/GameplayTagsEditor/Source/GameplayTagsEditor/Private/SAddNewGameplayTagSourceWidget.h
@@ -2,14 +2,10 @@
 
 #pragma once
 
-#include "CoreMinimal.h"
-#include "GameplayTagsManager.h"
-#include "Input/Reply.h"
-#include "Widgets/SWidget.h"
 #include "Widgets/SCompoundWidget.h"
-#include "Widgets/DeclarativeSyntaxSupport.h"
-#include "Widgets/Input/SEditableTextBox.h"
-#include "Widgets/Input/SComboBox.h"
+
+class SEditableTextBox;
+template <typename OptionType> class SComboBox;
 
 /** Widget allowing the user to create new gameplay tags */
 class SAddNewGameplayTagSourceWidget : public SCompoundWidget
diff --git a/Engine/Plugins/Editor/GameplayTagsEditor/Source/GameplayTagsEditor/Private/SAddNewGameplayTagWidget.cpp b/Engine/Plugins/Editor/GameplayTagsEditor/Source/GameplayTagsEditor/Private/SAddNewGameplayTagWidget.cpp
--- a/Engine/Plugins/Editor/GameplayTagsEditor/Source/GameplayTagsEditor/Private/SAddNewGameplayTagWidget.cpp
+++ b/Engine/Plugins/Editor/GameplayTagsEditor/Source/GameplayTagsEditor/Private/SAddNewGameplayTagWidget.cpp
@@ -2,12 +2,15 @@
 
 #include "SAddNewGameplayTagWidget.h"
 #include "DetailLayoutBuilder.h"
-#include "GameplayTagsSettings.h"
+#include "Framework/Views/TableViewMetadata.h"
 #include "GameplayTagsEditorModule.h"
+#include "GameplayTagsManager.h"
 #include "GameplayTagsModule.h"
+#include "Misc/Paths.h"
 #include "Widgets/Input/SButton.h"
 #include "Misc/MessageDialog.h"
-#include "Widgets/Images/SImage.h"
+#include "Widgets/Input/SComboBox.h"
+#include "Widgets/Input/SEditableTextBox.h"
 
 #define LOCTEXT_NAMESPACE "AddNewGameplayTagWidget"
 
diff --git a/Engine/Plugins/Editor/GameplayTagsEditor/Source/GameplayTagsEditor/Private/SAddNewGameplayTagWidget.h b/Engine/Plugins/Editor/GameplayTagsEditor/Source/GameplayTagsEditor/Private/SAddNewGameplayTagWidget.h
--- a/Engine/Plugins/Editor/GameplayTagsEditor/Source/GameplayTagsEditor/Private/SAddNewGameplayTagWidget.h
+++ b/Engine/Plugins/Editor/GameplayTagsEditor/Source/GameplayTagsEditor/Private/SAddNewGameplayTagWidget.h
@@ -2,14 +2,11 @@
 
 #pragma once
 
-#include "CoreMinimal.h"
-#include "GameplayTagsManager.h"
-#include "Input/Reply.h"
-#include "Widgets/SWidget.h"
 #include "Widgets/SCompoundWidget.h"
-#include "Widgets/DeclarativeSyntaxSupport.h"
-#include "Widgets/Input/SEditableTextBox.h"
-#include "Widgets/Input/SComboBox.h"
+
+class SEditableTextBox;
+namespace ETextCommit { enum Type : int; }
+template <typename OptionType> class SComboBox;
 
 /** Widget allowing the user to create new gameplay tags */
 class SAddNewGameplayTagWidget : public SCompoundWidget
diff --git a/Engine/Plugins/Editor/GameplayTagsEditor/Source/GameplayTagsEditor/Private/SAddNewRestrictedGameplayTagWidget.cpp b/Engine/Plugins/Editor/GameplayTagsEditor/Source/GameplayTagsEditor/Private/SAddNewRestrictedGameplayTagWidget.cpp
--- a/Engine/Plugins/Editor/GameplayTagsEditor/Source/GameplayTagsEditor/Private/SAddNewRestrictedGameplayTagWidget.cpp
+++ b/Engine/Plugins/Editor/GameplayTagsEditor/Source/GameplayTagsEditor/Private/SAddNewRestrictedGameplayTagWidget.cpp
@@ -2,13 +2,16 @@
 
 #include "SAddNewRestrictedGameplayTagWidget.h"
 #include "DetailLayoutBuilder.h"
-#include "GameplayTagsSettings.h"
+#include "Framework/Views/TableViewMetadata.h"
 #include "GameplayTagsEditorModule.h"
+#include "GameplayTagsManager.h"
 #include "GameplayTagsModule.h"
-#include "SGameplayTagWidget.h"
+#include "HAL/PlatformProcess.h"
 #include "Widgets/Input/SButton.h"
+#include "Widgets/Input/SComboBox.h"
 #include "Widgets/Notifications/SNotificationList.h"
 #include "Framework/Notifications/NotificationManager.h"
+#include "Widgets/Input/SEditableTextBox.h"
 
 #define LOCTEXT_NAMESPACE "AddNewRestrictedGameplayTagWidget"
 
diff --git a/Engine/Plugins/Editor/GameplayTagsEditor/Source/GameplayTagsEditor/Private/SAddNewRestrictedGameplayTagWidget.h b/Engine/Plugins/Editor/GameplayTagsEditor/Source/GameplayTagsEditor/Private/SAddNewRestrictedGameplayTagWidget.h
--- a/Engine/Plugins/Editor/GameplayTagsEditor/Source/GameplayTagsEditor/Private/SAddNewRestrictedGameplayTagWidget.h
+++ b/Engine/Plugins/Editor/GameplayTagsEditor/Source/GameplayTagsEditor/Private/SAddNewRestrictedGameplayTagWidget.h
@@ -2,16 +2,13 @@
 
 #pragma once
 
-#include "CoreMinimal.h"
-#include "GameplayTagsManager.h"
-#include "Input/Reply.h"
-#include "Widgets/SWidget.h"
 #include "Widgets/SCompoundWidget.h"
-#include "Widgets/DeclarativeSyntaxSupport.h"
-#include "Widgets/Input/SEditableTextBox.h"
-#include "Widgets/Input/SComboBox.h"
-#include "Widgets/Input/SCheckBox.h"
-#include "Widgets/Notifications/SNotificationList.h"
+
+class SCheckBox;
+class SEditableTextBox;
+class SNotificationItem;
+namespace ETextCommit { enum Type : int; }
+template <typename OptionType> class SComboBox;
 
 /** Widget allowing the user to create new restricted gameplay tags */
 class SAddNewRestrictedGameplayTagWidget : public SCompoundWidget
diff --git a/Engine/Plugins/Editor/GameplayTagsEditor/Source/GameplayTagsEditor/Private/SGameplayTagContainerGraphPin.cpp b/Engine/Plugins/Editor/GameplayTagsEditor/Source/GameplayTagsEditor/Private/SGameplayTagContainerGraphPin.cpp
--- a/Engine/Plugins/Editor/GameplayTagsEditor/Source/GameplayTagsEditor/Private/SGameplayTagContainerGraphPin.cpp
+++ b/Engine/Plugins/Editor/GameplayTagsEditor/Source/GameplayTagsEditor/Private/SGameplayTagContainerGraphPin.cpp
@@ -1,7 +1,10 @@
 // Copyright Epic Games, Inc. All Rights Reserved.
 
 #include "SGameplayTagContainerGraphPin.h"
+#include "EdGraph/EdGraphSchema.h"
 #include "GameplayTagPinUtilities.h"
+#include "GameplayTagContainer.h"
+#include "Widgets/SBoxPanel.h"
 
 #define LOCTEXT_NAMESPACE "GameplayTagGraphPin"
 
diff --git a/Engine/Plugins/Editor/GameplayTagsEditor/Source/GameplayTagsEditor/Private/SGameplayTagGraphPin.cpp b/Engine/Plugins/Editor/GameplayTagsEditor/Source/GameplayTagsEditor/Private/SGameplayTagGraphPin.cpp
--- a/Engine/Plugins/Editor/GameplayTagsEditor/Source/GameplayTagsEditor/Private/SGameplayTagGraphPin.cpp
+++ b/Engine/Plugins/Editor/GameplayTagsEditor/Source/GameplayTagsEditor/Private/SGameplayTagGraphPin.cpp
@@ -1,7 +1,11 @@
 // Copyright Epic Games, Inc. All Rights Reserved.
 
 #include "SGameplayTagGraphPin.h"
+#include "EdGraph/EdGraphSchema.h"
 #include "GameplayTagPinUtilities.h"
+#include "Framework/Views/TableViewMetadata.h"
+#include "GameplayTagContainer.h"
+#include "Widgets/Views/SListView.h"
 
 #define LOCTEXT_NAMESPACE "GameplayTagGraphPin"
 
diff --git a/Engine/Plugins/Editor/GameplayTagsEditor/Source/GameplayTagsEditor/Private/SGameplayTagGraphPin.h b/Engine/Plugins/Editor/GameplayTagsEditor/Source/GameplayTagsEditor/Private/SGameplayTagGraphPin.h
--- a/Engine/Plugins/Editor/GameplayTagsEditor/Source/GameplayTagsEditor/Private/SGameplayTagGraphPin.h
+++ b/Engine/Plugins/Editor/GameplayTagsEditor/Source/GameplayTagsEditor/Private/SGameplayTagGraphPin.h
@@ -2,12 +2,11 @@
 
 #pragma once
 
-#include "CoreMinimal.h"
 #include "KismetPins/SGraphPinStructInstance.h"
-#include "Widgets/Views/STableViewBase.h"
-#include "Widgets/Views/STableRow.h"
 #include "SGameplayTagWidget.h"
 
+template <typename ItemType> class SListView;
+
 /** Pin that represents a single gameplay tag, overrides the generic struct widget because tags have their own system for saving changes */
 class SGameplayTagGraphPin : public SGraphPinStructInstance
 {
diff --git a/Engine/Plugins/Editor/GameplayTagsEditor/Source/GameplayTagsEditor/Private/SGameplayTagQueryGraphPin.cpp b/Engine/Plugins/Editor/GameplayTagsEditor/Source/GameplayTagsEditor/Private/SGameplayTagQueryGraphPin.cpp
--- a/Engine/Plugins/Editor/GameplayTagsEditor/Source/GameplayTagsEditor/Private/SGameplayTagQueryGraphPin.cpp
+++ b/Engine/Plugins/Editor/GameplayTagsEditor/Source/GameplayTagsEditor/Private/SGameplayTagQueryGraphPin.cpp
@@ -1,8 +1,9 @@
 // Copyright Epic Games, Inc. All Rights Reserved.
 
 #include "SGameplayTagQueryGraphPin.h"
-#include "Widgets/SBoxPanel.h"
+#include "EdGraph/EdGraphSchema.h"
 #include "Widgets/Input/SComboButton.h"
+#include "GameplayTagContainer.h"
 #include "Widgets/Layout/SScaleBox.h"
 #include "ScopedTransaction.h"
 
diff --git a/Engine/Plugins/Editor/GameplayTagsEditor/Source/GameplayTagsEditor/Private/SGameplayTagQueryGraphPin.h b/Engine/Plugins/Editor/GameplayTagsEditor/Source/GameplayTagsEditor/Private/SGameplayTagQueryGraphPin.h
--- a/Engine/Plugins/Editor/GameplayTagsEditor/Source/GameplayTagsEditor/Private/SGameplayTagQueryGraphPin.h
+++ b/Engine/Plugins/Editor/GameplayTagsEditor/Source/GameplayTagsEditor/Private/SGameplayTagQueryGraphPin.h
@@ -2,9 +2,6 @@
 
 #pragma once
 
-#include "CoreMinimal.h"
-#include "Widgets/SWidget.h"
-#include "Widgets/DeclarativeSyntaxSupport.h"
 #include "SGraphPin.h"
 #include "SGameplayTagQueryWidget.h"
 
diff --git a/Engine/Plugins/Editor/GameplayTagsEditor/Source/GameplayTagsEditor/Private/SGameplayTagQueryWidget.cpp b/Engine/Plugins/Editor/GameplayTagsEditor/Source/GameplayTagsEditor/Private/SGameplayTagQueryWidget.cpp
--- a/Engine/Plugins/Editor/GameplayTagsEditor/Source/GameplayTagsEditor/Private/SGameplayTagQueryWidget.cpp
+++ b/Engine/Plugins/Editor/GameplayTagsEditor/Source/GameplayTagsEditor/Private/SGameplayTagQueryWidget.cpp
@@ -1,12 +1,11 @@
 // Copyright Epic Games, Inc. All Rights Reserved.
 
 #include "SGameplayTagQueryWidget.h"
+#include "DetailsViewArgs.h"
 #include "Modules/ModuleManager.h"
+#include "PropertyHandle.h"
 #include "Widgets/SBoxPanel.h"
-#include "Widgets/Layout/SBorder.h"
 #include "Widgets/Input/SButton.h"
-#include "Styling/AppStyle.h"
-#include "PropertyEditorModule.h"
 #include "IDetailsView.h"
 #include "GameplayTagsManager.h"
 
diff --git a/Engine/Plugins/Editor/GameplayTagsEditor/Source/GameplayTagsEditor/Private/SGameplayTagQueryWidget.h b/Engine/Plugins/Editor/GameplayTagsEditor/Source/GameplayTagsEditor/Private/SGameplayTagQueryWidget.h
--- a/Engine/Plugins/Editor/GameplayTagsEditor/Source/GameplayTagsEditor/Private/SGameplayTagQueryWidget.h
+++ b/Engine/Plugins/Editor/GameplayTagsEditor/Source/GameplayTagsEditor/Private/SGameplayTagQueryWidget.h
@@ -2,15 +2,10 @@
 
 #pragma once
 
-#include "CoreMinimal.h"
-#include "UObject/Object.h"
-#include "UObject/WeakObjectPtr.h"
-#include "Layout/Visibility.h"
-#include "Widgets/DeclarativeSyntaxSupport.h"
-#include "Input/Reply.h"
 #include "Widgets/SCompoundWidget.h"
-#include "GameplayTagContainer.h"
-#include "PropertyHandle.h"
+
+class IPropertyHandle;
+struct FGameplayTagQuery;
 
 class IDetailsView;
 struct FPropertyChangedEvent;
