[NO NON-PII COMMIT MESSAGE]

diff --git a/Engine/Plugins/Editor/GameplayTagsEditor/Source/GameplayTagsEditor/Private/GameplayTagCustomization.cpp b/Engine/Plugins/Editor/GameplayTagsEditor/Source/GameplayTagsEditor/Private/GameplayTagCustomization.cpp
--- a/Engine/Plugins/Editor/GameplayTagsEditor/Source/GameplayTagsEditor/Private/GameplayTagCustomization.cpp
+++ b/Engine/Plugins/Editor/GameplayTagsEditor/Source/GameplayTagsEditor/Private/GameplayTagCustomization.cpp
@@ -229,7 +229,10 @@ FGameplayTagCustomization::~FGameplayTagCustomization()
 		EditButton->SetIsOpen(false);
 	}
 
-	GEditor->UnregisterForUndo(this);
+	if (GEditor)
+	{
+		GEditor->UnregisterForUndo(this);
+	}
 }
 
 void FGameplayTagCustomization::BuildEditableContainerList()
