Changed GameplayTag viewer so GameplayTag sources are visible in all modes

diff --git a/Engine/Plugins/Editor/GameplayTagsEditor/Source/GameplayTagsEditor/Private/SGameplayTagWidget.cpp b/Engine/Plugins/Editor/GameplayTagsEditor/Source/GameplayTagsEditor/Private/SGameplayTagWidget.cpp
--- a/Engine/Plugins/Editor/GameplayTagsEditor/Source/GameplayTagsEditor/Private/SGameplayTagWidget.cpp
+++ b/Engine/Plugins/Editor/GameplayTagsEditor/Source/GameplayTagsEditor/Private/SGameplayTagWidget.cpp
@@ -456,22 +456,16 @@ TSharedRef<ITableRow> SGameplayTagWidget::OnGenerateRow(TSharedPtr<FGameplayTagN
 
 		if (Node.IsValid())
 		{
-			// Add Tag source if we're in management mode
-			if (GameplayTagUIMode == EGameplayTagUIMode::ManagementMode)
+			FName TagSource;
+			if (Node->bIsExplicitTag)
 			{
-				FName TagSource;
-
-				if (Node->bIsExplicitTag)
-				{
-					TagSource = Node->GetFirstSourceName();
-				}
-				else
-				{
-					TagSource = FName(TEXT("Implicit"));
-				}
-
-				TooltipString.Append(FString::Printf(TEXT(" (%s)"), *TagSource.ToString()));
+				TagSource = Node->GetFirstSourceName();
+			}
+			else
+			{
+				TagSource = FName(TEXT("Implicit"));
 			}
+			TooltipString.Append(FString::Printf(TEXT(" (%s)"), *TagSource.ToString()));
 
 			// tag comments
 			if (!Node->DevComment.IsEmpty())
