Better fix for crash when removing pins from a node with orphaned pins
#preflgiht 6418ce8c3f3d31c94abab654

diff --git a/Engine/Plugins/Editor/GameplayTagsEditor/Source/GameplayTagsEditor/Private/SGraphNode_MultiCompareGameplayTag.cpp b/Engine/Plugins/Editor/GameplayTagsEditor/Source/GameplayTagsEditor/Private/SGraphNode_MultiCompareGameplayTag.cpp
--- a/Engine/Plugins/Editor/GameplayTagsEditor/Source/GameplayTagsEditor/Private/SGraphNode_MultiCompareGameplayTag.cpp
+++ b/Engine/Plugins/Editor/GameplayTagsEditor/Source/GameplayTagsEditor/Private/SGraphNode_MultiCompareGameplayTag.cpp
@@ -120,7 +120,7 @@ FReply SGraphNode_MultiCompareGameplayTag::OnRemovePin()
 	FBlueprintEditorUtils::MarkBlueprintAsModified(CompareNode->GetBlueprint());
 
 	UpdateGraphNode();
-	GraphNode->GetGraph()->NotifyGraphChanged();
+	GraphNode->GetGraph()->NotifyNodeChanged(GraphNode);
 
 	return FReply::Handled();
 }
@@ -137,7 +137,7 @@ FReply SGraphNode_MultiCompareGameplayTag::OnAddPin()
 	FBlueprintEditorUtils::MarkBlueprintAsModified(CompareNode->GetBlueprint());
 
 	UpdateGraphNode();
-	GraphNode->GetGraph()->NotifyGraphChanged();
+	GraphNode->GetGraph()->NotifyNodeChanged(GraphNode);
 
 	return FReply::Handled();
 }
diff --git a/Engine/Source/Runtime/Engine/Private/EdGraph/EdGraph.cpp b/Engine/Source/Runtime/Engine/Private/EdGraph/EdGraph.cpp
--- a/Engine/Source/Runtime/Engine/Private/EdGraph/EdGraph.cpp
+++ b/Engine/Source/Runtime/Engine/Private/EdGraph/EdGraph.cpp
@@ -285,7 +285,11 @@ void UEdGraph::NotifyGraphChanged()
 
 void UEdGraph::NotifyNodeChanged(const UEdGraphNode* Node)
 {
-	NotifyGraphChanged();
+	FEdGraphEditAction Action;
+	Action.Graph = this;
+	Action.Action = GRAPHACTION_EditNode;
+	Action.Nodes.Add(Node);
+	NotifyGraphChanged(Action);
 }
 
 void UEdGraph::NotifyGraphChanged(const FEdGraphEditAction& InAction)
