Support for the FInheritedGameplayTags to update its CombinedTags immediately (rather than waiting until a Save).

diff --git a/Engine/Plugins/Editor/GameplayTagsEditor/Source/GameplayTagsEditor/Private/GameplayTagContainerCustomization.cpp b/Engine/Plugins/Editor/GameplayTagsEditor/Source/GameplayTagsEditor/Private/GameplayTagContainerCustomization.cpp
--- a/Engine/Plugins/Editor/GameplayTagsEditor/Source/GameplayTagsEditor/Private/GameplayTagContainerCustomization.cpp
+++ b/Engine/Plugins/Editor/GameplayTagsEditor/Source/GameplayTagsEditor/Private/GameplayTagContainerCustomization.cpp
@@ -40,6 +40,10 @@ void FGameplayTagContainerCustomization::CustomizeHeader(TSharedRef<class IPrope
 {
 	StructPropertyHandle = InStructPropertyHandle;
 
+	// We are registering to a global callback used to let us know when the GameplayTags have updated across the project.
+	// We are using this to piggyback on notifications that can be sent by other Properties to update ourselves. In particular, FInheritedTagContainer::CombinedTags.
+	UGameplayTagsManager::OnEditorRefreshGameplayTagTree.AddSP(this, &FGameplayTagContainerCustomization::RefreshTagList);
+
 	FSimpleDelegate OnTagContainerChanged = FSimpleDelegate::CreateSP(this, &FGameplayTagContainerCustomization::RefreshTagList);
 	StructPropertyHandle->SetOnPropertyValueChanged(OnTagContainerChanged);
 
diff --git a/Engine/Plugins/Runtime/GameplayAbilities/Source/GameplayAbilities/Private/GameplayEffect.cpp b/Engine/Plugins/Runtime/GameplayAbilities/Source/GameplayAbilities/Private/GameplayEffect.cpp
--- a/Engine/Plugins/Runtime/GameplayAbilities/Source/GameplayAbilities/Private/GameplayEffect.cpp
+++ b/Engine/Plugins/Runtime/GameplayAbilities/Source/GameplayAbilities/Private/GameplayEffect.cpp
@@ -5162,6 +5162,11 @@ void FInheritedTagContainer::UpdateInheritedTagProperties(const FInheritedTagCon
 			CombinedTags.AddTag(*Itr);
 		}
 	}
+
+#if WITH_EDITOR
+	// Notify any open editor windows that we should refresh these properties, specifically for refreshing CombinedTags.
+	UGameplayTagsManager::OnEditorRefreshGameplayTagTree.Broadcast();
+#endif
 }
 
 void FInheritedTagContainer::PostInitProperties()
diff --git a/Engine/Plugins/Runtime/GameplayAbilities/Source/GameplayAbilitiesEditor/Private/InheritableGameplayTagContainerDetails.cpp b/Engine/Plugins/Runtime/GameplayAbilities/Source/GameplayAbilitiesEditor/Private/InheritableGameplayTagContainerDetails.cpp
deleted file mode 100644
diff --git a/Engine/Plugins/Runtime/GameplayAbilities/Source/GameplayAbilitiesEditor/Private/InheritableGameplayTagContainerDetails.h b/Engine/Plugins/Runtime/GameplayAbilities/Source/GameplayAbilitiesEditor/Private/InheritableGameplayTagContainerDetails.h
deleted file mode 100644
