GameplayTags: Improved gameplay tag UI
- Added common widget SGameplayTagChip, which can be used to display a gameplaytag
- Added common widget SGameplayTagCombo which can be used to edit a Gameplay Tag
- Added common widget SGameplayTagContainerCombo which can be used to edit a Gameplay Tag Container
- Added common widget SGameplayTagQueryEntryBox which can be used to edit a Gameplay Tag Query
- Improved appearance and functionalityof the gameplay tag picker
- Moved tag management into a separate window
- Update details customization and graph pins to use the new widgets

diff --git a/Engine/Plugins/Editor/GameplayTagsEditor/Source/GameplayTagsEditor/GameplayTagsEditor.Build.cs b/Engine/Plugins/Editor/GameplayTagsEditor/Source/GameplayTagsEditor/GameplayTagsEditor.Build.cs
--- a/Engine/Plugins/Editor/GameplayTagsEditor/Source/GameplayTagsEditor/GameplayTagsEditor.Build.cs
+++ b/Engine/Plugins/Editor/GameplayTagsEditor/Source/GameplayTagsEditor/GameplayTagsEditor.Build.cs
@@ -22,6 +22,7 @@ namespace UnrealBuildTool.Rules
 					"Core",
 					"CoreUObject",
 					"Engine",
+					"ApplicationCore",
 					"AssetTools",
 					"AssetRegistry",
 					"GameplayTags",
@@ -40,6 +41,7 @@ namespace UnrealBuildTool.Rules
 					"SourceControl",
 					"ToolMenus",
 					"SettingsEditor",
+					"ToolWidgets",
 				}
 			);
 
diff --git a/Engine/Plugins/Editor/GameplayTagsEditor/Source/GameplayTagsEditor/Internal/SAddNewGameplayTagWidget.h b/Engine/Plugins/Editor/GameplayTagsEditor/Source/GameplayTagsEditor/Internal/SAddNewGameplayTagWidget.h
--- a/Engine/Plugins/Editor/GameplayTagsEditor/Source/GameplayTagsEditor/Internal/SAddNewGameplayTagWidget.h
+++ b/Engine/Plugins/Editor/GameplayTagsEditor/Source/GameplayTagsEditor/Internal/SAddNewGameplayTagWidget.h
@@ -7,6 +7,7 @@
 class SEditableTextBox;
 namespace ETextCommit { enum Type : int; }
 template <typename OptionType> class SComboBox;
+class SNotificationItem;
 
 /** Widget allowing the user to create new gameplay tags */
 class SAddNewGameplayTagWidget : public SCompoundWidget
@@ -31,7 +32,7 @@ public:
 		SLATE_ARGUMENT(FString, NewTagName) // String that will initially populate the New Tag Name field
 	SLATE_END_ARGS();
 
-	GAMEPLAYTAGSEDITOR_API virtual ~SAddNewGameplayTagWidget();
+	GAMEPLAYTAGSEDITOR_API virtual ~SAddNewGameplayTagWidget() override;
 
 	GAMEPLAYTAGSEDITOR_API virtual void Tick(const FGeometry& AllottedGeometry, const double InCurrentTime, const float InDeltaTime) override;
 
@@ -46,6 +47,9 @@ public:
 	/** Begins the process of adding a subtag to a parent tag */
 	void AddSubtagFromParent(const FString& ParentTagName, const FName& ParentTagSource);
 
+	/** Begins the process of adding a duplicate of existing tag */
+	void AddDuplicate(const FString& ParentTagName, const FName& ParentTagSource);
+
 	/** Resets all input fields */
 	void Reset(EResetType ResetType);
 
@@ -66,6 +70,9 @@ private:
 	/** Callback for when the Add New Tag button is pressed */
 	FReply OnAddNewTagButtonPressed();
 
+	/** Return true of the tag data is valid and can be added. */
+	bool IsTagDataValid();
+	
 	/** Creates a new gameplay tag and adds it to the INI files based on the widget's stored parameters */
 	void CreateNewGameplayTag();
 
@@ -90,11 +97,11 @@ private:
 	/** The name of the next gameplay tag to create */
 	TSharedPtr<SEditableTextBox> TagNameTextBox;
 
-	/** The comment to asign to the next gameplay tag to create*/
+	/** The comment to asign to the next gameplay tag to create */
 	TSharedPtr<SEditableTextBox> TagCommentTextBox;
 
-	/** The INI file where the next gameplay tag will be creatd */
-	TSharedPtr<SComboBox<TSharedPtr<FName> > > TagSourcesComboBox;
+	/** The INI file where the next gameplay tag will be created */
+	TSharedPtr<SComboBox<TSharedPtr<FName>>> TagSourcesComboBox;
 
 	/** Callback for when a new gameplay tag has been added to the INI files */
 	FOnGameplayTagAdded OnGameplayTagAdded;
@@ -102,10 +109,15 @@ private:
 	/** Callback to see if the gameplay tag is valid. This should be used for any specialized rules that are not covered by IsValidGameplayTagString */
 	FIsValidTag IsValidTag;
 
-	bool bAddingNewTag;
+	/** Guards against the window closing when it loses focus due to source control checking out a file */
+	bool bAddingNewTag = false;
 
 	/** Tracks if this widget should get keyboard focus */
-	bool bShouldGetKeyboardFocus;
+	bool bShouldGetKeyboardFocus = false;
 
+	/** Cached from NewTagName argument. */
 	FString DefaultNewName;
+
+	/** Last error notification trying to create a new tag. */
+	TSharedPtr<SNotificationItem> NotificationItem;
 };
diff --git a/Engine/Plugins/Editor/GameplayTagsEditor/Source/GameplayTagsEditor/Private/AssetTypeActions_GameplayTagAssetBase.cpp b/Engine/Plugins/Editor/GameplayTagsEditor/Source/GameplayTagsEditor/Private/AssetTypeActions_GameplayTagAssetBase.cpp
--- a/Engine/Plugins/Editor/GameplayTagsEditor/Source/GameplayTagsEditor/Private/AssetTypeActions_GameplayTagAssetBase.cpp
+++ b/Engine/Plugins/Editor/GameplayTagsEditor/Source/GameplayTagsEditor/Private/AssetTypeActions_GameplayTagAssetBase.cpp
@@ -4,7 +4,7 @@
 #include "Framework/Application/SlateApplication.h"
 
 #include "GameplayTagContainer.h"
-#include "SGameplayTagWidget.h"
+#include "SGameplayTagPicker.h"
 #include "Interfaces/IMainFrameModule.h"
 #include "ToolMenuSection.h"
 
@@ -16,62 +16,125 @@ FAssetTypeActions_GameplayTagAssetBase::FAssetTypeActions_GameplayTagAssetBase(F
 
 void FAssetTypeActions_GameplayTagAssetBase::GetActions(const TArray<UObject*>& InObjects, FToolMenuSection& Section)
 {
-	TArray<UObject*> ContainerObjectOwners;
-	TArray<FGameplayTagContainer*> Containers;
+	TArray<UObject*> Objects;
+	TArray<FGameplayTagContainer> Containers;
+	
 	for (int32 ObjIdx = 0; ObjIdx < InObjects.Num(); ++ObjIdx)
 	{
-		UObject* CurObj = InObjects[ObjIdx];
-		if (CurObj)
+		if (UObject* CurObj = InObjects[ObjIdx])
 		{
-			FStructProperty* StructProp = FindFProperty<FStructProperty>(CurObj->GetClass(), OwnedGameplayTagPropertyName);
-			if(StructProp != NULL)
+			const FStructProperty* Property = FindFProperty<FStructProperty>(CurObj->GetClass(), OwnedGameplayTagPropertyName);
+			if(Property != nullptr)
 			{
-				ContainerObjectOwners.Add(CurObj);
-				Containers.Add(StructProp->ContainerPtrToValuePtr<FGameplayTagContainer>(CurObj));
+				const FGameplayTagContainer& Container = *Property->ContainerPtrToValuePtr<FGameplayTagContainer>(CurObj); 
+				Objects.Add(CurObj);
+				Containers.Add(Container);
 			}
 		}
 	}
 
-	ensure(Containers.Num() == ContainerObjectOwners.Num());
-	if (Containers.Num() > 0 && (Containers.Num() == ContainerObjectOwners.Num()))
+	if (Containers.Num() > 0)
 	{
 		Section.AddMenuEntry(
 			"GameplayTags_Edit",
 			LOCTEXT("GameplayTags_Edit", "Edit Gameplay Tags..."),
 			LOCTEXT("GameplayTags_EditToolTip", "Opens the Gameplay Tag Editor."),
 			FSlateIcon(),
-			FUIAction(FExecuteAction::CreateSP(this, &FAssetTypeActions_GameplayTagAssetBase::OpenGameplayTagEditor, ContainerObjectOwners, Containers), FCanExecuteAction()));
+			FUIAction(FExecuteAction::CreateSP(this, &FAssetTypeActions_GameplayTagAssetBase::OpenGameplayTagEditor, Objects, Containers), FCanExecuteAction()));
 	}
 }
 
-void FAssetTypeActions_GameplayTagAssetBase::OpenGameplayTagEditor(TArray<UObject*> Objects, TArray<FGameplayTagContainer*> Containers)
+void FAssetTypeActions_GameplayTagAssetBase::OpenGameplayTagEditor(TArray<UObject*> Objects, TArray<FGameplayTagContainer> Containers) const
 {
-	TArray<SGameplayTagWidget::FEditableGameplayTagContainerDatum> EditableContainers;
-	for (int32 ObjIdx = 0; ObjIdx < Objects.Num() && ObjIdx < Containers.Num(); ++ObjIdx)
+	if (!Objects.Num() || !Containers.Num())
 	{
-		EditableContainers.Add(SGameplayTagWidget::FEditableGameplayTagContainerDatum(Objects[ObjIdx], Containers[ObjIdx]));
+		return;
+	}
+
+	check(Objects.Num() == Containers.Num());
+
+	for (UObject* Object : Objects)
+	{
+		check (IsValid(Object));
+		Object->SetFlags(RF_Transactional);
 	}
 
 	FText Title;
 	FText AssetName;
 
-	const int32 NumAssets = EditableContainers.Num();
+	const int32 NumAssets = Containers.Num();
 	if (NumAssets > 1)
 	{
 		AssetName = FText::Format( LOCTEXT("AssetTypeActions_GameplayTagAssetBaseMultipleAssets", "{0} Assets"), FText::AsNumber( NumAssets ) );
 		Title = FText::Format( LOCTEXT("AssetTypeActions_GameplayTagAssetBaseEditorTitle", "Tag Editor: Owned Gameplay Tags: {0}"), AssetName );
 	}
-	else if (NumAssets > 0 && EditableContainers[0].TagContainerOwner.IsValid())
+	else
 	{
-		AssetName = FText::FromString( EditableContainers[0].TagContainerOwner->GetName() );
+		AssetName = FText::FromString(GetNameSafe(Objects[0]));
 		Title = FText::Format( LOCTEXT("AssetTypeActions_GameplayTagAssetBaseEditorTitle", "Tag Editor: Owned Gameplay Tags: {0}"), AssetName );
 	}
 
 	TSharedPtr<SWindow> Window = SNew(SWindow)
 		.Title(Title)
-		.ClientSize(FVector2D(600, 400))
+		.ClientSize(FVector2D(500, 600))
 		[
-			SNew(SGameplayTagWidget, EditableContainers)
+			SNew(SGameplayTagPicker)
+			.TagContainers(Containers)
+			.MaxHeight(0) // unbounded
+			.OnRefreshTagContainers_Lambda([Objects, OwnedGameplayTagPropertyName = OwnedGameplayTagPropertyName](SGameplayTagPicker& TagPicker)
+			{
+				// Refresh tags from objects, this is called e.g. on post undo/redo. 
+				TArray<FGameplayTagContainer> Containers;
+				for (UObject* Object : Objects)
+				{
+					// Adding extra entry even if the object has gone invalid to keep the container count the same as object count.
+					FGameplayTagContainer& NewContainer = Containers.AddDefaulted_GetRef();
+					if (IsValid(Object))
+					{
+						const FStructProperty* Property = FindFProperty<FStructProperty>(Object->GetClass(), OwnedGameplayTagPropertyName);
+						if (Property != nullptr)
+						{
+							NewContainer = *Property->ContainerPtrToValuePtr<FGameplayTagContainer>(Object); 
+						}
+					}
+				}
+				TagPicker.SetTagContainers(Containers);
+			})
+			.OnTagChanged_Lambda([Objects, OwnedGameplayTagPropertyName = OwnedGameplayTagPropertyName](const TArray<FGameplayTagContainer>& TagContainers)
+			{
+				// Sanity check that our arrays are in sync.
+				if (Objects.Num() != TagContainers.Num())
+				{
+					return;
+				}
+				
+				for (int32 Index = 0; Index < Objects.Num(); Index++)
+				{
+					UObject* Object = Objects[Index];
+					if (!IsValid(Object))
+					{
+						continue;
+					}
+
+					FStructProperty* Property = FindFProperty<FStructProperty>(Object->GetClass(), OwnedGameplayTagPropertyName);
+					if (!Property)
+					{
+						continue;
+					}
+
+					Object->Modify();
+					FGameplayTagContainer& Container = *Property->ContainerPtrToValuePtr<FGameplayTagContainer>(Object); 
+
+					FEditPropertyChain PropertyChain;
+					PropertyChain.AddHead(Property);
+					Object->PreEditChange(PropertyChain);
+
+					Container = TagContainers[Index];
+					
+					FPropertyChangedEvent PropertyEvent(Property);
+					Object->PostEditChangeProperty(PropertyEvent);
+				}
+			})
 		];
 
 	IMainFrameModule& MainFrameModule = FModuleManager::LoadModuleChecked<IMainFrameModule>(TEXT("MainFrame"));
diff --git a/Engine/Plugins/Editor/GameplayTagsEditor/Source/GameplayTagsEditor/Private/GameplayTagContainerCustomization.cpp b/Engine/Plugins/Editor/GameplayTagsEditor/Source/GameplayTagsEditor/Private/GameplayTagContainerCustomization.cpp
--- a/Engine/Plugins/Editor/GameplayTagsEditor/Source/GameplayTagsEditor/Private/GameplayTagContainerCustomization.cpp
+++ b/Engine/Plugins/Editor/GameplayTagsEditor/Source/GameplayTagsEditor/Private/GameplayTagContainerCustomization.cpp
@@ -1,464 +1,124 @@
 // Copyright Epic Games, Inc. All Rights Reserved.
 
 #include "GameplayTagContainerCustomization.h"
-#include "AssetRegistry/AssetIdentifier.h"
-#include "Widgets/Input/SComboButton.h"
-
-
-
-#include "Editor.h"
 #include "DetailWidgetRow.h"
-#include "Framework/Application/SlateApplication.h"
-#include "Framework/Views/TableViewMetadata.h"
 #include "GameplayTagsEditorModule.h"
 #include "GameplayTagsManager.h"
+#include "Widgets/Input/SHyperlink.h"
+#include "SGameplayTagContainerCombo.h"
 #include "HAL/PlatformApplicationMisc.h"
-#include "Layout/WidgetPath.h"
-#include "SGameplayTagWidget.h"
 #include "ScopedTransaction.h"
-#include "Widgets/Input/SHyperlink.h"
-#include "EditorFontGlyphs.h"
-#include "Framework/MultiBox/MultiBoxBuilder.h"
+#include "GameplayTagEditorUtilities.h"
+#include "SGameplayTagPicker.h"
 
 #define LOCTEXT_NAMESPACE "GameplayTagContainerCustomization"
 
 TSharedRef<IPropertyTypeCustomization> FGameplayTagContainerCustomizationPublic::MakeInstance()
 {
-	return MakeInstanceWithOptions({});
+	return MakeShareable(new FGameplayTagContainerCustomization());
 }
 
+// Deprecated version.
 TSharedRef<IPropertyTypeCustomization> FGameplayTagContainerCustomizationPublic::MakeInstanceWithOptions(const FGameplayTagContainerCustomizationOptions& Options)
 {
-	return MakeShareable(new FGameplayTagContainerCustomization(Options));
+	return MakeShareable(new FGameplayTagContainerCustomization());
 }
 
-FGameplayTagContainerCustomization::FGameplayTagContainerCustomization(const FGameplayTagContainerCustomizationOptions& InOptions):
-	Options(InOptions)
-{}
+FGameplayTagContainerCustomization::FGameplayTagContainerCustomization()
+{
+}
 
-void FGameplayTagContainerCustomization::CustomizeHeader(TSharedRef<class IPropertyHandle> InStructPropertyHandle, class FDetailWidgetRow& HeaderRow, IPropertyTypeCustomizationUtils& StructCustomizationUtils)
+void FGameplayTagContainerCustomization::CustomizeHeader(TSharedRef<IPropertyHandle> InStructPropertyHandle, FDetailWidgetRow& HeaderRow, IPropertyTypeCustomizationUtils& StructCustomizationUtils)
 {
 	StructPropertyHandle = InStructPropertyHandle;
 
-	// We are registering to a global callback used to let us know when the GameplayTags have updated across the project.
-	// We are using this to piggyback on notifications that can be sent by other Properties to update ourselves. In particular, FInheritedTagContainer::CombinedTags.
-	UGameplayTagsManager::OnEditorRefreshGameplayTagTree.AddSP(this, &FGameplayTagContainerCustomization::RefreshTagList);
-
-	FSimpleDelegate OnTagContainerChanged = FSimpleDelegate::CreateSP(this, &FGameplayTagContainerCustomization::RefreshTagList);
-	StructPropertyHandle->SetOnPropertyValueChanged(OnTagContainerChanged);
-
-	FUIAction SearchForReferencesAction(FExecuteAction::CreateSP(this, &FGameplayTagContainerCustomization::OnWholeContainerSearchForReferences));
-
 	HeaderRow
 		.NameContent()
 		[
 			StructPropertyHandle->CreatePropertyNameWidget()
 		]
 		.ValueContent()
-		.MaxDesiredWidth(512)
-		[
-			SNew(SHorizontalBox)
-			+ SHorizontalBox::Slot()
-			.AutoWidth()
-			.VAlign(VAlign_Center)
-			[
-				SNew(SVerticalBox)
-				+ SVerticalBox::Slot()
-				.AutoHeight()
-				[
-					SAssignNew(EditButton, SComboButton)
-					.OnGetMenuContent(this, &FGameplayTagContainerCustomization::GetListContent)
-					.OnMenuOpenChanged(this, &FGameplayTagContainerCustomization::OnGameplayTagListMenuOpenStateChanged)
-					.ContentPadding(FMargin(2.0f, 2.0f))
-					.MenuPlacement(MenuPlacement_BelowAnchor)
-					.ButtonContent()
-					[
-						SNew(STextBlock)
-						.Text(LOCTEXT("GameplayTagContainerCustomization_Edit", "Edit..."))
-					]
-				]
-				+ SVerticalBox::Slot()
-				.AutoHeight()
-				[
-					SNew(SButton)
-					.IsEnabled(!StructPropertyHandle->IsEditConst())
-					.Text(LOCTEXT("GameplayTagContainerCustomization_Clear", "Clear All"))
-					.OnClicked(this, &FGameplayTagContainerCustomization::OnClearAllButtonClicked)
-					.Visibility(this, &FGameplayTagContainerCustomization::GetClearAllVisibility)
-				]
-			]
-			+ SHorizontalBox::Slot()
-			.AutoWidth()
-			[
-				SNew(SBorder)
-				.Padding(4.0f)
-				.Visibility(this, &FGameplayTagContainerCustomization::GetTagsListVisibility)
-				[
-					ActiveTags()
-				]
-			]
-		]
-		.AddCustomContextMenuAction(SearchForReferencesAction,
-			LOCTEXT("WholeContainerSearchForReferences", "Search For References"),
-			LOCTEXT("WholeContainerSearchForReferencesTooltip", "Find referencers that reference *any* of the tags in this container"),
-			FSlateIcon());
-
-	GEditor->RegisterForUndo(this);
-}
-
-TSharedRef<SWidget> FGameplayTagContainerCustomization::ActiveTags()
-{	
-	RefreshTagList();
-	
-	SAssignNew( TagListView, SListView<TSharedPtr<FGameplayTag>> )
-	.ListItemsSource(&TagList)
-	.SelectionMode(ESelectionMode::None)
-	.OnGenerateRow(this, &FGameplayTagContainerCustomization::MakeListViewWidget);
-
-	return TagListView->AsShared();
-}
-
-void FGameplayTagContainerCustomization::RefreshTagList()
-{
-	// Build the set of tags on any instance, collapsing common tags together
-	TSet<FGameplayTag> CurrentTagSet;
-	if (StructPropertyHandle.IsValid())
-	{
-		SGameplayTagWidget::EnumerateEditableTagContainersFromPropertyHandle(StructPropertyHandle.ToSharedRef(), [&CurrentTagSet](const SGameplayTagWidget::FEditableGameplayTagContainerDatum& EditableTagContainer)
-		{
-			if (const FGameplayTagContainer* Container = EditableTagContainer.TagContainer)
-			{
-				for (auto It = Container->CreateConstIterator(); It; ++It)
-				{
-					CurrentTagSet.Add(*It);
-				}
-			}
-			return true;
-		});
-	}
-
-	// Convert the set into pointers for the combo
-	TagList.Empty(CurrentTagSet.Num());
-	for (const FGameplayTag& CurrentTag : CurrentTagSet)
-	{
-		TagList.Add(MakeShared<FGameplayTag>(CurrentTag));
-	}
-	TagList.StableSort([](const TSharedPtr<FGameplayTag>& One, const TSharedPtr<FGameplayTag>& Two)
-	{
-		return *One < *Two;
-	});
-
-	// Refresh the slate list
-	if( TagListView.IsValid() )
-	{
-		TagListView->RequestListRefresh();
-	}
-}
-
-TSharedRef<ITableRow> FGameplayTagContainerCustomization::MakeListViewWidget(TSharedPtr<FGameplayTag> Item, const TSharedRef<STableViewBase>& OwnerTable)
-{
-	TSharedPtr<SWidget> TagItem;
-
-	const FString TagName = Item->ToString();
-	if (UGameplayTagsManager::Get().ShowGameplayTagAsHyperLinkEditor(TagName))
-	{
-		TagItem = SNew(SHyperlink)
-			.Text(FText::FromString(TagName))
-			.OnNavigate(this, &FGameplayTagContainerCustomization::OnTagDoubleClicked, *Item.Get());
-	}
-	else
-	{
-		TagItem = SNew(STextBlock)
-			.Text(FText::FromString(TagName));
-	}
-
-	return SNew( STableRow< TSharedPtr<FString> >, OwnerTable )
-	[
-		SNew(SBorder)
-		.OnMouseButtonDown(this, &FGameplayTagContainerCustomization::OnSingleTagMouseButtonPressed, TagName)
-		.Padding(0.0f)
-		.BorderImage(FAppStyle::GetBrush("NoBorder"))
+		.VAlign(VAlign_Center)
 		[
-			SNew(SHorizontalBox)
-			+ SHorizontalBox::Slot()
-			.AutoWidth()
-			.VAlign(VAlign_Center)
-			.Padding(0,0,2,0)
+			SNew(SBox)
+			.Padding(FMargin(0,2,0,1))
 			[
-				SNew(SButton)
-				.IsEnabled(!StructPropertyHandle->IsEditConst())
-				.ContentPadding(FMargin(0))
-				.ButtonStyle(FAppStyle::Get(), "FlatButton.Danger")
-				.ForegroundColor(FSlateColor::UseForeground())
-				.OnClicked(this, &FGameplayTagContainerCustomization::OnRemoveTagClicked, *Item.Get())
-				[
-					SNew(STextBlock)
-					.Font(FAppStyle::Get().GetFontStyle("FontAwesome.9"))
-					.Text(FEditorFontGlyphs::Times)
-				]
-			]
-
-			+ SHorizontalBox::Slot()
-			.AutoWidth()
-			.VAlign(VAlign_Center)
-			[
-				TagItem.ToSharedRef()
+				SNew(SGameplayTagContainerCombo)
+				.PropertyHandle(StructPropertyHandle)
 			]
 		]
-	];
-}
-
-FReply FGameplayTagContainerCustomization::OnSingleTagMouseButtonPressed(const FGeometry& MyGeometry, const FPointerEvent& MouseEvent, FString TagName)
-{
-	if (MouseEvent.IsMouseButtonDown(EKeys::RightMouseButton))
-	{
-		FMenuBuilder MenuBuilder(/*bShouldCloseWindowAfterMenuSelection=*/ true, /*CommandList=*/ nullptr);
-
-		FUIAction SearchForReferencesAction(FExecuteAction::CreateSP(this, &FGameplayTagContainerCustomization::OnSingleTagSearchForReferences, TagName));
-
-		MenuBuilder.BeginSection(NAME_None, FText::Format(LOCTEXT("SingleTagMenuHeading", "Tag Actions ({0})"), FText::AsCultureInvariant(TagName)));
-		MenuBuilder.AddMenuEntry(
-			LOCTEXT("SingleTagSearchForReferences", "Search For References"),
-			FText::Format(LOCTEXT("SingleTagSearchForReferencesTooltip", "Find references to the tag {0}"), FText::AsCultureInvariant(TagName)),
-			FSlateIcon(),
-			SearchForReferencesAction);
-		
-		if (StructPropertyHandle)
-		{
-			FUIAction CopyAction = FUIAction(FExecuteAction::CreateSP(this, &FGameplayTagContainerCustomization::OnCopyTag, TagName));
-			FUIAction PasteAction = FUIAction(FExecuteAction::CreateSP(this, &FGameplayTagContainerCustomization::OnPasteTag),FCanExecuteAction::CreateSP(this, &FGameplayTagContainerCustomization::CanPaste));
-			
-			if (CopyAction.IsBound() && PasteAction.IsBound())
-			{
-				FMenuEntryParams CopyContentParams;
-				CopyContentParams.LabelOverride = NSLOCTEXT("PropertyView", "CopyProperty", "Copy");
-				CopyContentParams.ToolTipOverride = FText::Format(LOCTEXT("SingleTagCopyTooltip", "Copy tag {0}"), FText::AsCultureInvariant(TagName));
-				CopyContentParams.IconOverride = FSlateIcon(FCoreStyle::Get().GetStyleSetName(), "GenericCommands.Copy");
-				CopyContentParams.DirectActions = CopyAction;
-				MenuBuilder.AddMenuEntry(CopyContentParams);
-
-				FMenuEntryParams PasteContentParams;
-				PasteContentParams.LabelOverride = NSLOCTEXT("PropertyView", "PasteProperty", "Paste");
-				PasteContentParams.ToolTipOverride = FText::Format(LOCTEXT("SingleTagPasteTooltip", "Paste tag {0}"), FText::AsCultureInvariant(TagName));
-				PasteContentParams.IconOverride = FSlateIcon(FCoreStyle::Get().GetStyleSetName(), "GenericCommands.Paste");
-				PasteContentParams.DirectActions = PasteAction;
-				MenuBuilder.AddMenuEntry(PasteContentParams);
-			}
-		}
-		
-		MenuBuilder.EndSection();
-
-		// Spawn context menu
-		FWidgetPath WidgetPath = MouseEvent.GetEventPath() != nullptr ? *MouseEvent.GetEventPath() : FWidgetPath();
-		FSlateApplication::Get().PushMenu(TagListView.ToSharedRef(), WidgetPath, MenuBuilder.MakeWidget(), MouseEvent.GetScreenSpacePosition(), FPopupTransitionEffect(FPopupTransitionEffect::ContextMenu));
-
-		return FReply::Handled();
-	}
-
-	return FReply::Unhandled();
-}
-
-void FGameplayTagContainerCustomization::OnSingleTagSearchForReferences(FString TagName)
-{
-	FName TagFName(*TagName, FNAME_Find);
-	if (FEditorDelegates::OnOpenReferenceViewer.IsBound() && !TagFName.IsNone())
-	{
-		TArray<FAssetIdentifier> AssetIdentifiers;
-		AssetIdentifiers.Emplace(FGameplayTag::StaticStruct(), TagFName);
-		FEditorDelegates::OnOpenReferenceViewer.Broadcast(AssetIdentifiers, FReferenceViewerParams());
-	}
+	.PasteAction(FUIAction(
+	FExecuteAction::CreateSP(this, &FGameplayTagContainerCustomization::OnPasteTag),
+		FCanExecuteAction::CreateSP(this, &FGameplayTagContainerCustomization::CanPasteTag)));
 }
 
-void FGameplayTagContainerCustomization::OnWholeContainerSearchForReferences()
+void FGameplayTagContainerCustomization::OnPasteTag() const
 {
-	if (FEditorDelegates::OnOpenReferenceViewer.IsBound())
+	if (!StructPropertyHandle.IsValid())
 	{
-		TArray<FAssetIdentifier> AssetIdentifiers;
-		AssetIdentifiers.Reserve(TagList.Num());
-		for (const TSharedPtr<FGameplayTag>& TagPtr : TagList)
-		{
-			if (TagPtr->IsValid())
-			{
-				AssetIdentifiers.Emplace(FGameplayTag::StaticStruct(), TagPtr->GetTagName());
-			}
-		}
-
-		FEditorDelegates::OnOpenReferenceViewer.Broadcast(AssetIdentifiers, FReferenceViewerParams());
+		return;
 	}
-}
-
-void FGameplayTagContainerCustomization::OnTagDoubleClicked(FGameplayTag Tag)
-{
-	UGameplayTagsManager::Get().NotifyGameplayTagDoubleClickedEditor(Tag.ToString());
-}
+	
+	FString PastedText;
+	FPlatformApplicationMisc::ClipboardPaste(PastedText);
+	bool bHandled = false;
 
-FReply FGameplayTagContainerCustomization::OnRemoveTagClicked(FGameplayTag Tag)
-{
-	TArray<FString> NewValues;
-	SGameplayTagWidget::EnumerateEditableTagContainersFromPropertyHandle(StructPropertyHandle.ToSharedRef(), [&Tag, &NewValues](const SGameplayTagWidget::FEditableGameplayTagContainerDatum& EditableTagContainer)
+	// Try to paste single tag
+	const FGameplayTag PastedTag = UE::GameplayTags::EditorUtilities::GameplayTagTryImportText(PastedText);
+	if (PastedTag.IsValid())
 	{
-		FGameplayTagContainer TagContainerCopy;
-		if (const FGameplayTagContainer* Container = EditableTagContainer.TagContainer)
+		TArray<FString> NewValues;
+		SGameplayTagPicker::EnumerateEditableTagContainersFromPropertyHandle(StructPropertyHandle.ToSharedRef(), [&NewValues, PastedTag](const FGameplayTagContainer& EditableTagContainer)
 		{
-			TagContainerCopy = *Container;
-		}
-		TagContainerCopy.RemoveTag(Tag);
+			FGameplayTagContainer TagContainerCopy = EditableTagContainer;
+			TagContainerCopy.AddTag(PastedTag);
 
-		NewValues.Add(TagContainerCopy.ToString());
-		return true;
-	});
+			NewValues.Add(TagContainerCopy.ToString());
+			return true;
+		});
 
-	{
-		FScopedTransaction Transaction(LOCTEXT("RemoveGameplayTagFromContainer", "Remove Gameplay Tag"));
+		FScopedTransaction Transaction(LOCTEXT("GameplayTagContainerCustomization_PasteTag", "Paste Gameplay Tag"));
 		StructPropertyHandle->SetPerObjectValues(NewValues);
+		bHandled = true;
 	}
 
-	RefreshTagList();
-
-	return FReply::Handled();
-}
-
-TSharedRef<SWidget> FGameplayTagContainerCustomization::GetListContent()
-{
-	if (!StructPropertyHandle.IsValid() || StructPropertyHandle->GetProperty() == nullptr)
-	{
-		return SNullWidget::NullWidget;
-	}
-
-	FString Categories = UGameplayTagsManager::Get().GetCategoriesMetaFromPropertyHandle(StructPropertyHandle);
-
-	bool bReadOnly = StructPropertyHandle->IsEditConst();
-
-	TSharedRef<SGameplayTagWidget> TagWidget = SNew(SGameplayTagWidget, TArray<SGameplayTagWidget::FEditableGameplayTagContainerDatum>()) // empty container list as built from StructPropertyHandle
-		.Filter(Categories)
-		.ReadOnly(bReadOnly)
-		.TagContainerName(StructPropertyHandle->GetPropertyDisplayName().ToString())
-		.OnTagChanged(this, &FGameplayTagContainerCustomization::RefreshTagList)
-		.PropertyHandle(StructPropertyHandle)
-		.ForceHideAddNewTag(Options.bForceHideAddTag)
-		.ForceHideAddNewTagSource(Options.bForceHideAddTagSource);
-
-	LastTagWidget = TagWidget;
-
-	return SNew(SVerticalBox)
-		+ SVerticalBox::Slot()
-		.AutoHeight()
-		.MaxHeight(400)
-		[
-			TagWidget
-		];
-}
-
-void FGameplayTagContainerCustomization::OnGameplayTagListMenuOpenStateChanged(bool bIsOpened)
-{
-	if (bIsOpened)
+	// Try to paste a container
+	if (!bHandled)
 	{
-		TSharedPtr<SGameplayTagWidget> TagWidget = LastTagWidget.Pin();
-		if (TagWidget.IsValid())
+		const FGameplayTagContainer PastedTagContainer = UE::GameplayTags::EditorUtilities::GameplayTagContainerTryImportText(PastedText);
+		if (PastedTagContainer.IsValid())
 		{
-			EditButton->SetMenuContentWidgetToFocus(TagWidget->GetWidgetToFocusOnOpen());
+			// From property
+			FScopedTransaction Transaction(LOCTEXT("GameplayTagContainerCustomization_PasteTagContainer", "Paste Gameplay Tag Container"));
+			StructPropertyHandle->SetValueFromFormattedString(PastedText);
+			bHandled = true;
 		}
 	}
 }
 
-FReply FGameplayTagContainerCustomization::OnClearAllButtonClicked()
-{
-	{
-		FScopedTransaction Transaction(LOCTEXT("GameplayTagContainerCustomization_RemoveAllTags", "Remove All Gameplay Tags"));
-		StructPropertyHandle->SetValueFromFormattedString(FGameplayTagContainer().ToString());
-	}
-
-	RefreshTagList();
-
-	return FReply::Handled();
-}
-
-EVisibility FGameplayTagContainerCustomization::GetClearAllVisibility() const
-{
-	return TagList.Num() > 0 ? EVisibility::Visible : EVisibility::Collapsed;
-}
-
-EVisibility FGameplayTagContainerCustomization::GetTagsListVisibility() const
-{
-	return TagList.Num() > 0 ? EVisibility::Visible : EVisibility::Collapsed;
-}
-
-void FGameplayTagContainerCustomization::PostUndo( bool bSuccess )
-{
-	if( bSuccess )
-	{
-		RefreshTagList();
-	}
-}
-
-void FGameplayTagContainerCustomization::PostRedo( bool bSuccess )
-{
-	if( bSuccess )
-	{
-		RefreshTagList();
-	}
-}
-
-FGameplayTagContainerCustomization::~FGameplayTagContainerCustomization()
+bool FGameplayTagContainerCustomization::CanPasteTag() const
 {
-	// Forcibly close the popup to avoid crashes later
-	if (EditButton.IsValid() && EditButton->IsOpen())
+	if (!StructPropertyHandle.IsValid())
 	{
-		EditButton->SetIsOpen(false);
+		return false;
 	}
 
-	GEditor->UnregisterForUndo(this);
-}
-
-void FGameplayTagContainerCustomization::OnCopyTag(FString TagName)
-{
-	FPlatformApplicationMisc::ClipboardCopy(*TagName);
-}
-
-void FGameplayTagContainerCustomization::OnPasteTag()
-{
-	FString TagName;
-	FPlatformApplicationMisc::ClipboardPaste(TagName);
+	FString PastedText;
+	FPlatformApplicationMisc::ClipboardPaste(PastedText);
 	
-	const FGameplayTag Tag = FGameplayTag::RequestGameplayTag(FName(TagName), false);
-
-	if (Tag.IsValid())
+	const FGameplayTag PastedTag = UE::GameplayTags::EditorUtilities::GameplayTagTryImportText(PastedText);
+	if (PastedTag.IsValid())
 	{
-		TArray<FString> NewValues;
-		SGameplayTagWidget::EnumerateEditableTagContainersFromPropertyHandle(StructPropertyHandle.ToSharedRef(), [&Tag, &NewValues](const SGameplayTagWidget::FEditableGameplayTagContainerDatum& EditableTagContainer)
-		{
-			FGameplayTagContainer TagContainerCopy;
-			if (const FGameplayTagContainer* Container = EditableTagContainer.TagContainer)
-			{
-				TagContainerCopy = *Container;
-			}
-			TagContainerCopy.AddTag(Tag);
-
-			NewValues.Add(TagContainerCopy.ToString());
-			return true;
-		});
-
-		{
-			FScopedTransaction Transaction(LOCTEXT("PasteGameplayTagFromContainer", "Paste Gameplay Tag"));
-			StructPropertyHandle->SetPerObjectValues(NewValues);
-		}
-
-		RefreshTagList();
+		return true;
 	}
-}
 
-bool FGameplayTagContainerCustomization::CanPaste()
-{
-	FString TagName;
-	FPlatformApplicationMisc::ClipboardPaste(TagName);
-
-	if (TagName.Len() < NAME_SIZE)
+	const FGameplayTagContainer PastedTagContainer = UE::GameplayTags::EditorUtilities::GameplayTagContainerTryImportText(PastedText);
+	if (PastedTagContainer.IsValid())
 	{
-		const FGameplayTag Tag = FGameplayTag::RequestGameplayTag(FName(TagName), false);
-		return Tag.IsValid();
+		return true;
 	}
-	
-    return false;
+
+	return false;
 }
 
 #undef LOCTEXT_NAMESPACE
diff --git a/Engine/Plugins/Editor/GameplayTagsEditor/Source/GameplayTagsEditor/Private/GameplayTagContainerCustomization.h b/Engine/Plugins/Editor/GameplayTagsEditor/Source/GameplayTagsEditor/Private/GameplayTagContainerCustomization.h
--- a/Engine/Plugins/Editor/GameplayTagsEditor/Source/GameplayTagsEditor/Private/GameplayTagContainerCustomization.h
+++ b/Engine/Plugins/Editor/GameplayTagsEditor/Source/GameplayTagsEditor/Private/GameplayTagContainerCustomization.h
@@ -3,82 +3,30 @@
 #pragma once
 
 #include "IPropertyTypeCustomization.h"
-#include "GameplayTagContainerCustomizationOptions.h"
-#include "EditorUndoClient.h"
-#include "Widgets/Views/SListView.h"
 
 struct FGameplayTag;
-
 class IPropertyHandle;
+class FDetailWidgetRow;
+class IDetailChildrenBuilder;
+class IPropertyTypeCustomizationUtils;
 
 /** Customization for the gameplay tag container struct */
-class FGameplayTagContainerCustomization : public IPropertyTypeCustomization, public FEditorUndoClient
+class FGameplayTagContainerCustomization : public IPropertyTypeCustomization
 {
 public:
-
-	FGameplayTagContainerCustomization(const FGameplayTagContainerCustomizationOptions& InOptions);
-	~FGameplayTagContainerCustomization();
+	FGameplayTagContainerCustomization();
 
 	/** Overridden to show an edit button to launch the gameplay tag editor */
-	virtual void CustomizeHeader(TSharedRef<class IPropertyHandle> StructPropertyHandle, class FDetailWidgetRow& HeaderRow, IPropertyTypeCustomizationUtils& StructCustomizationUtils) override;
+	virtual void CustomizeHeader(TSharedRef<IPropertyHandle> StructPropertyHandle, FDetailWidgetRow& HeaderRow, IPropertyTypeCustomizationUtils& StructCustomizationUtils) override;
 	
 	/** Overridden to do nothing */
-	virtual void CustomizeChildren(TSharedRef<IPropertyHandle> InStructPropertyHandle, class IDetailChildrenBuilder& ChildBuilder, IPropertyTypeCustomizationUtils& StructCustomizationUtils) override {}
-
-	//~ Begin FEditorUndoClient Interface
-	virtual void PostUndo(bool bSuccess) override;
-	virtual void PostRedo(bool bSuccess) override;
-	//~ End FEditorUndoClient Interface
+	virtual void CustomizeChildren(TSharedRef<IPropertyHandle> InStructPropertyHandle, IDetailChildrenBuilder& ChildBuilder, IPropertyTypeCustomizationUtils& StructCustomizationUtils) override {}
 
 private:
-
-	/** Called when the clear all button is clicked; Clears all selected tags in the container*/
-	FReply OnClearAllButtonClicked();
-
-	/** Returns the visibility of the "clear all" button (collapsed when there are no tags) */
-	EVisibility GetClearAllVisibility() const;
-
-	/** Returns the visibility of the tags list box (collapsed when there are no tags) */
-	EVisibility GetTagsListVisibility() const;
-
-	/** Returns the selected tags list widget*/
-	TSharedRef<SWidget> ActiveTags();
-
-	/** Updates the list of selected tags*/
-	void RefreshTagList();
-
-	/** On Generate Row Delegate */
-	TSharedRef<ITableRow> MakeListViewWidget(TSharedPtr<FGameplayTag> Item, const TSharedRef<STableViewBase>& OwnerTable);
-
-	/** Callback function to create content for the combo button. */
-	TSharedRef<SWidget> GetListContent();
-
-	void OnGameplayTagListMenuOpenStateChanged(bool bIsOpened);
-
 	/** Cached property handle */
 	TSharedPtr<IPropertyHandle> StructPropertyHandle;
 
-	/** List of tags selected in the tag containers */
-	TArray< TSharedPtr<FGameplayTag> > TagList;
-
-	/** The TagList, kept as a member so we can update it later */
-	TSharedPtr<SListView<TSharedPtr<FGameplayTag>>> TagListView;
-
-	TSharedPtr<class SComboButton> EditButton;
-
-	TWeakPtr<class SGameplayTagWidget> LastTagWidget;
-
-	FGameplayTagContainerCustomizationOptions Options;
-
-	void OnTagDoubleClicked(FGameplayTag Tag);
-	FReply OnRemoveTagClicked(FGameplayTag Tag);
-
-	FReply OnSingleTagMouseButtonPressed(const FGeometry& MyGeometry, const FPointerEvent& MouseEvent, FString TagName);
-	void OnSingleTagSearchForReferences(FString TagName);
-	void OnWholeContainerSearchForReferences();
-
-	void OnCopyTag(FString TagName);
-	void OnPasteTag();
-	bool CanPaste();
+	void OnPasteTag() const;
+	bool CanPasteTag() const;
 };
 
diff --git a/Engine/Plugins/Editor/GameplayTagsEditor/Source/GameplayTagsEditor/Private/GameplayTagCustomization.cpp b/Engine/Plugins/Editor/GameplayTagsEditor/Source/GameplayTagsEditor/Private/GameplayTagCustomization.cpp
--- a/Engine/Plugins/Editor/GameplayTagsEditor/Source/GameplayTagsEditor/Private/GameplayTagCustomization.cpp
+++ b/Engine/Plugins/Editor/GameplayTagsEditor/Source/GameplayTagsEditor/Private/GameplayTagCustomization.cpp
@@ -1,43 +1,43 @@
 // Copyright Epic Games, Inc. All Rights Reserved.
 
 #include "GameplayTagCustomization.h"
-#include "AssetRegistry/AssetIdentifier.h"
-#include "Widgets/Input/SComboButton.h"
-
-#include "Editor.h"
+#include "PropertyHandle.h"
 #include "DetailWidgetRow.h"
 #include "GameplayTagsManager.h"
 #include "IDetailChildrenBuilder.h"
 #include "GameplayTagsEditorModule.h"
 #include "Widgets/Input/SHyperlink.h"
+#include "SGameplayTagCombo.h"
+#include "SGameplayTagPicker.h"
 
 #define LOCTEXT_NAMESPACE "GameplayTagCustomization"
 
+//---------------------------------------------------------
+// FGameplayTagCustomizationPublic
+//---------------------------------------------------------
+
 TSharedRef<IPropertyTypeCustomization> FGameplayTagCustomizationPublic::MakeInstance()
 {
-	return MakeInstanceWithOptions({});
+	return MakeShareable(new FGameplayTagCustomization());
 }
 
+// Deprecated version.
 TSharedRef<IPropertyTypeCustomization> FGameplayTagCustomizationPublic::MakeInstanceWithOptions(const FGameplayTagCustomizationOptions& Options)
 {
-	return MakeShareable(new FGameplayTagCustomization(Options));
+	return MakeShareable(new FGameplayTagCustomization());
 }
 
-FGameplayTagCustomization::FGameplayTagCustomization(const FGameplayTagCustomizationOptions& InOptions):
-	Options(InOptions)
-{}
+//---------------------------------------------------------
+// FGameplayTagCustomization
+//---------------------------------------------------------
 
-void FGameplayTagCustomization::CustomizeHeader(TSharedRef<class IPropertyHandle> InStructPropertyHandle, class FDetailWidgetRow& HeaderRow, IPropertyTypeCustomizationUtils& StructCustomizationUtils)
+FGameplayTagCustomization::FGameplayTagCustomization()
 {
-	TagContainer = MakeShareable(new FGameplayTagContainer);
-	StructPropertyHandle = InStructPropertyHandle;
-
-	FSimpleDelegate OnTagChanged = FSimpleDelegate::CreateSP(this, &FGameplayTagCustomization::OnPropertyValueChanged);
-	StructPropertyHandle->SetOnPropertyValueChanged(OnTagChanged);
-
-	BuildEditableContainerList();
+}
 
-	FUIAction SearchForReferencesAction(FExecuteAction::CreateSP(this, &FGameplayTagCustomization::OnSearchForReferences));
+void FGameplayTagCustomization::CustomizeHeader(TSharedRef<IPropertyHandle> InStructPropertyHandle, FDetailWidgetRow& HeaderRow, IPropertyTypeCustomizationUtils& StructCustomizationUtils)
+{
+	StructPropertyHandle = InStructPropertyHandle;
 
 	HeaderRow
 	.NameContent()
@@ -45,254 +45,46 @@ void FGameplayTagCustomization::CustomizeHeader(TSharedRef<class IPropertyHandle
 		StructPropertyHandle->CreatePropertyNameWidget()
 	]
 	.ValueContent()
-	.MaxDesiredWidth(512)
 	[
-		SNew(SHorizontalBox)
-		+ SHorizontalBox::Slot()
-		.AutoWidth()
+		SNew(SBox)
+		.Padding(FMargin(0,2,0,1))
 		[
-			SAssignNew(EditButton, SComboButton)
-			.OnGetMenuContent(this, &FGameplayTagCustomization::GetListContent)
-			.OnMenuOpenChanged(this, &FGameplayTagCustomization::OnGameplayTagListMenuOpenStateChanged)
-			.ContentPadding(FMargin(2.0f, 2.0f))
-			.MenuPlacement(MenuPlacement_BelowAnchor)
-			.ButtonContent()
-			[
-				SNew(STextBlock)
-				.Text(LOCTEXT("GameplayTagCustomization_Edit", "Edit"))
-			]
+			SNew(SGameplayTagCombo)
+			.PropertyHandle(StructPropertyHandle)
 		]
-		
-		+ SHorizontalBox::Slot()
-		.AutoWidth()
-		[
-			SNew(SBorder)
-			.Visibility(this, &FGameplayTagCustomization::GetVisibilityForTagTextBlockWidget, true)
-			.Padding(4.0f)
-			[
-				SNew(STextBlock)
-				.Text(this, &FGameplayTagCustomization::SelectedTag)
-			]
-		]
-
-		+ SHorizontalBox::Slot()
-		.AutoWidth()
-		[
-			SNew(SBorder)
-			.Visibility(this, &FGameplayTagCustomization::GetVisibilityForTagTextBlockWidget, false)
-			.Padding(4.0f)
-			[
-				SNew(SHyperlink)
-				.Text(this, &FGameplayTagCustomization::SelectedTag)
-				.OnNavigate( this, &FGameplayTagCustomization::OnTagDoubleClicked)
-			]
-		]
-	]
-	.AddCustomContextMenuAction(SearchForReferencesAction,
-		LOCTEXT("FGameplayTagCustomization_SearchForReferences", "Search For References"),
-		LOCTEXT("FGameplayTagCustomization_SearchForReferencesTooltip", "Find references for this tag"),
-		FSlateIcon());
-
-	GEditor->RegisterForUndo(this);
-}
-
-void FGameplayTagCustomization::OnTagDoubleClicked()
-{
-	UGameplayTagsManager::Get().NotifyGameplayTagDoubleClickedEditor(TagName);
-}
-
-void FGameplayTagCustomization::OnSearchForReferences()
-{
-	FName TagFName(*TagName, FNAME_Find);
-	if (FEditorDelegates::OnOpenReferenceViewer.IsBound() && !TagFName.IsNone())
-	{
-		TArray<FAssetIdentifier> AssetIdentifiers;
-		AssetIdentifiers.Add(FAssetIdentifier(FGameplayTag::StaticStruct(), TagFName));
-		FEditorDelegates::OnOpenReferenceViewer.Broadcast(AssetIdentifiers, FReferenceViewerParams());
-	}
-}
-
-EVisibility FGameplayTagCustomization::GetVisibilityForTagTextBlockWidget(bool ForTextWidget) const
-{
-	return (UGameplayTagsManager::Get().ShowGameplayTagAsHyperLinkEditor(TagName) ^ ForTextWidget) ? EVisibility::Visible : EVisibility::Collapsed;
+	];
 }
 
-TSharedRef<SWidget> FGameplayTagCustomization::GetListContent()
-{
-	BuildEditableContainerList();
-	
-	FString Categories = UGameplayTagsManager::Get().GetCategoriesMetaFromPropertyHandle(StructPropertyHandle);
-
-	bool bReadOnly = StructPropertyHandle->IsEditConst();
-
-	TSharedRef<SGameplayTagWidget> TagWidget =
-		SNew(SGameplayTagWidget, EditableContainers)
-		.Filter(Categories)
-		.ReadOnly(bReadOnly)
-		.TagContainerName(StructPropertyHandle->GetPropertyDisplayName().ToString())
-		.MultiSelect(false)
-		.OnTagChanged(this, &FGameplayTagCustomization::OnTagChanged)
-		.PropertyHandle(StructPropertyHandle)
-		.ForceHideAddNewTag(Options.bForceHideAddTag)
-		.ForceHideAddNewTagSource(Options.bForceHideAddTagSource);
-
-	LastTagWidget = TagWidget;
-
-	return SNew(SVerticalBox)
-		+ SVerticalBox::Slot()
-		.AutoHeight()
-		.MaxHeight(400)
-		[
-			TagWidget
-		];
-}
-
-void FGameplayTagCustomization::OnGameplayTagListMenuOpenStateChanged(bool bIsOpened)
-{
-	if (bIsOpened)
-	{
-		TSharedPtr<SGameplayTagWidget> TagWidget = LastTagWidget.Pin();
-		if (TagWidget.IsValid())
-		{
-			EditButton->SetMenuContentWidgetToFocus(TagWidget->GetWidgetToFocusOnOpen());
-		}
-	}
-}
-
-void FGameplayTagCustomization::OnPropertyValueChanged()
-{
-	TagName = TEXT("");
-	if (StructPropertyHandle.IsValid() && StructPropertyHandle->GetProperty() && EditableContainers.Num() > 0)
-	{
-		TArray<void*> RawStructData;
-		StructPropertyHandle->AccessRawData(RawStructData);
-		if (RawStructData.Num() > 0)
-		{
-			FGameplayTag* Tag = (FGameplayTag*)(RawStructData[0]);
-			FGameplayTagContainer* Container = EditableContainers[0].TagContainer;			
-			if (Tag && Container)
-			{
-				Container->Reset();
-				Container->AddTag(*Tag);
-				TagName = Tag->ToString();
-			}			
-		}
-	}
-}
-
-void FGameplayTagCustomization::OnTagChanged()
-{
-	TagName = TEXT("");
-	if (StructPropertyHandle.IsValid() && StructPropertyHandle->GetProperty() && EditableContainers.Num() > 0)
-	{
-		TArray<void*> RawStructData;
-		StructPropertyHandle->AccessRawData(RawStructData);
-		if (RawStructData.Num() > 0)
-		{
-			FGameplayTag* Tag = (FGameplayTag*)(RawStructData[0]);			
-
-			// Update Tag from the one selected from list
-			FGameplayTagContainer* Container = EditableContainers[0].TagContainer;
-			if (Tag && Container)
-			{
-				for (auto It = Container->CreateConstIterator(); It; ++It)
-				{
-					*Tag = *It;
-					TagName = It->ToString();
-				}
-			}
-		}
-	}
-}
-
-void FGameplayTagCustomization::PostUndo(bool bSuccess)
-{
-	if (bSuccess && !StructPropertyHandle.IsValid())
-	{
-		OnTagChanged();
-	}
-}
-
-void FGameplayTagCustomization::PostRedo(bool bSuccess)
-{
-	if (bSuccess && !StructPropertyHandle.IsValid())
-	{
-		OnTagChanged();
-	}
-}
-
-FGameplayTagCustomization::~FGameplayTagCustomization()
-{
-	// Forcibly close the popup to avoid crashes later
-	if (EditButton.IsValid() && EditButton->IsOpen())
-	{
-		EditButton->SetIsOpen(false);
-	}
-
-	if (GEditor)
-	{
-		GEditor->UnregisterForUndo(this);
-	}
-}
-
-void FGameplayTagCustomization::BuildEditableContainerList()
-{
-	EditableContainers.Empty();
-
-	if(StructPropertyHandle.IsValid() && StructPropertyHandle->GetProperty())
-	{
-		TArray<void*> RawStructData;
-		StructPropertyHandle->AccessRawData(RawStructData);
-
-		if (RawStructData.Num() > 0)
-		{
-			FGameplayTag* Tag = (FGameplayTag*)(RawStructData[0]);
-			if (Tag && Tag->IsValid())
-			{
-				TagName = Tag->ToString();
-				TagContainer->AddTag(*Tag);
-			}
-		}
-
-		EditableContainers.Add(SGameplayTagWidget::FEditableGameplayTagContainerDatum(nullptr, TagContainer.Get()));
-	}
-}
-
-FText FGameplayTagCustomization::SelectedTag() const
-{
-	return FText::FromString(*TagName);
-}
+//---------------------------------------------------------
+// FGameplayTagCreationWidgetHelperDetails
+//---------------------------------------------------------
 
 TSharedRef<IPropertyTypeCustomization> FGameplayTagCreationWidgetHelperDetails::MakeInstance()
 {
 	return MakeShareable(new FGameplayTagCreationWidgetHelperDetails());
 }
 
-void FGameplayTagCreationWidgetHelperDetails::CustomizeHeader(TSharedRef<IPropertyHandle> StructPropertyHandle, class FDetailWidgetRow& HeaderRow, IPropertyTypeCustomizationUtils& StructCustomizationUtils)
+void FGameplayTagCreationWidgetHelperDetails::CustomizeHeader(TSharedRef<IPropertyHandle> StructPropertyHandle, FDetailWidgetRow& HeaderRow, IPropertyTypeCustomizationUtils& StructCustomizationUtils)
 {
-
+	HeaderRow.WholeRowContent()[ SNullWidget::NullWidget ];
 }
 
-void FGameplayTagCreationWidgetHelperDetails::CustomizeChildren(TSharedRef<IPropertyHandle> StructPropertyHandle, class IDetailChildrenBuilder& StructBuilder, IPropertyTypeCustomizationUtils& StructCustomizationUtils)
+void FGameplayTagCreationWidgetHelperDetails::CustomizeChildren(TSharedRef<IPropertyHandle> StructPropertyHandle, IDetailChildrenBuilder& StructBuilder, IPropertyTypeCustomizationUtils& StructCustomizationUtils)
 {
-	FString FilterString = UGameplayTagsManager::Get().GetCategoriesMetaFromPropertyHandle(StructPropertyHandle);
-	const float MaxPropertyWidth = 480.0f;
-	const float MaxPropertyHeight = 240.0f;
+	const FString FilterString = UGameplayTagsManager::Get().GetCategoriesMetaFromPropertyHandle(StructPropertyHandle);
+	constexpr float MaxPropertyWidth = 480.0f;
+	constexpr float MaxPropertyHeight = 240.0f;
 
 	StructBuilder.AddCustomRow(NSLOCTEXT("GameplayTagReferenceHelperDetails", "NewTag", "NewTag"))
 		.ValueContent()
 		.MaxDesiredWidth(MaxPropertyWidth)
 		[
-			SAssignNew(TagWidget, SGameplayTagWidget, TArray<SGameplayTagWidget::FEditableGameplayTagContainerDatum>())
+			SAssignNew(TagWidget, SGameplayTagPicker)
 			.Filter(FilterString)
-		.NewTagName(FilterString)
-		.MultiSelect(false)
-		.GameplayTagUIMode(EGameplayTagUIMode::ManagementMode)
-		.MaxHeight(MaxPropertyHeight)
-		.NewTagControlsInitiallyExpanded(true)
-		//.OnTagChanged(this, &FGameplayTagsSettingsCustomization::OnTagChanged)
+			.MultiSelect(false)
+			.GameplayTagPickerMode (EGameplayTagPickerMode::ManagementMode)
+			.MaxHeight(MaxPropertyHeight)
 		];
-
 }
 
 #undef LOCTEXT_NAMESPACE
diff --git a/Engine/Plugins/Editor/GameplayTagsEditor/Source/GameplayTagsEditor/Private/GameplayTagCustomization.h b/Engine/Plugins/Editor/GameplayTagsEditor/Source/GameplayTagsEditor/Private/GameplayTagCustomization.h
--- a/Engine/Plugins/Editor/GameplayTagsEditor/Source/GameplayTagsEditor/Private/GameplayTagCustomization.h
+++ b/Engine/Plugins/Editor/GameplayTagsEditor/Source/GameplayTagsEditor/Private/GameplayTagCustomization.h
@@ -3,74 +3,35 @@
 #pragma once
 
 #include "IPropertyTypeCustomization.h"
-#include "GameplayTagCustomizationOptions.h"
-#include "SGameplayTagWidget.h"
-#include "EditorUndoClient.h"
+#include "GameplayTagContainer.h"
 
 struct EVisibility;
 
 class IPropertyHandle;
+class FDetailWidgetRow;
+class IDetailChildrenBuilder;
+class SGameplayTagPicker;
 
 /** Customization for the gameplay tag struct */
-class FGameplayTagCustomization : public IPropertyTypeCustomization, public FEditorUndoClient
+class FGameplayTagCustomization : public IPropertyTypeCustomization
 {
 public:
 
-	FGameplayTagCustomization(const FGameplayTagCustomizationOptions& InOptions);
-	~FGameplayTagCustomization();
+	FGameplayTagCustomization();
 
 	/** Overridden to show an edit button to launch the gameplay tag editor */
-	virtual void CustomizeHeader(TSharedRef<class IPropertyHandle> StructPropertyHandle, class FDetailWidgetRow& HeaderRow, IPropertyTypeCustomizationUtils& StructCustomizationUtils) override;
+	virtual void CustomizeHeader(TSharedRef<IPropertyHandle> StructPropertyHandle, FDetailWidgetRow& HeaderRow, IPropertyTypeCustomizationUtils& StructCustomizationUtils) override;
 	
 	/** Overridden to do nothing */
-	virtual void CustomizeChildren(TSharedRef<IPropertyHandle> InStructPropertyHandle, class IDetailChildrenBuilder& ChildBuilder, IPropertyTypeCustomizationUtils& StructCustomizationUtils) override {}
-
-	//~ Begin FEditorUndoClient Interface
-	virtual void PostUndo(bool bSuccess) override;
-	virtual void PostRedo(bool bSuccess) override;
-	//~ End FEditorUndoClient Interface
+	virtual void CustomizeChildren(TSharedRef<IPropertyHandle> InStructPropertyHandle, IDetailChildrenBuilder& ChildBuilder, IPropertyTypeCustomizationUtils& StructCustomizationUtils) override {}
 
 private:
 
-	/** Updates the selected tag*/
-	void OnTagChanged();
-
-	/** Updates the selected tag*/
-	void OnPropertyValueChanged();
-
-	/** Build Editable Container */
-	void BuildEditableContainerList();
-
-	/** Callback function to create content for the combo button. */
-	TSharedRef<SWidget> GetListContent();
-
-	void OnGameplayTagListMenuOpenStateChanged(bool bIsOpened);
-
-	/** Returns Tag name currently selected*/
-	FText SelectedTag() const;
-
 	/** Cached property handle */
 	TSharedPtr<IPropertyHandle> StructPropertyHandle;
 
-	/** Tag Container used for the GameplayTagWidget. */
-	TSharedPtr<FGameplayTagContainer> TagContainer;
-
-	/** Editable Container for holding our tag */
-	TArray<SGameplayTagWidget::FEditableGameplayTagContainerDatum> EditableContainers;
-
-	/** Tag name selected*/
-	FString TagName;
-
-	void OnTagDoubleClicked();
-	void OnSearchForReferences();
-
-	EVisibility GetVisibilityForTagTextBlockWidget(bool ForTextWidget) const;
-
-	TSharedPtr<class SComboButton> EditButton;
-
-	TWeakPtr<class SGameplayTagWidget> LastTagWidget;
-
-	FGameplayTagCustomizationOptions Options;
+	/** Edited Tag */
+	FGameplayTag Tag;
 };
 
 /** Customization for FGameplayTagCreationWidgetHelper showing an add tag button */
@@ -80,7 +41,8 @@ public:
 	static TSharedRef<IPropertyTypeCustomization> MakeInstance();
 
 	/** IPropertyTypeCustomization interface */
-	virtual void CustomizeHeader(TSharedRef<class IPropertyHandle> StructPropertyHandle, class FDetailWidgetRow& HeaderRow, IPropertyTypeCustomizationUtils& StructCustomizationUtils) override;
-	virtual void CustomizeChildren(TSharedRef<class IPropertyHandle> StructPropertyHandle, class IDetailChildrenBuilder& StructBuilder, IPropertyTypeCustomizationUtils& StructCustomizationUtils) override;
-	TSharedPtr<class SGameplayTagWidget> TagWidget;
+	virtual void CustomizeHeader(TSharedRef<IPropertyHandle> StructPropertyHandle, FDetailWidgetRow& HeaderRow, IPropertyTypeCustomizationUtils& StructCustomizationUtils) override;
+	virtual void CustomizeChildren(TSharedRef<IPropertyHandle> StructPropertyHandle, IDetailChildrenBuilder& StructBuilder, IPropertyTypeCustomizationUtils& StructCustomizationUtils) override;
+	
+	TSharedPtr<SGameplayTagPicker> TagWidget;
 };
diff --git a/Engine/Plugins/Editor/GameplayTagsEditor/Source/GameplayTagsEditor/Private/GameplayTagQueryCustomization.cpp b/Engine/Plugins/Editor/GameplayTagsEditor/Source/GameplayTagsEditor/Private/GameplayTagQueryCustomization.cpp
--- a/Engine/Plugins/Editor/GameplayTagsEditor/Source/GameplayTagsEditor/Private/GameplayTagQueryCustomization.cpp
+++ b/Engine/Plugins/Editor/GameplayTagsEditor/Source/GameplayTagsEditor/Private/GameplayTagQueryCustomization.cpp
@@ -1,261 +1,23 @@
 // Copyright Epic Games, Inc. All Rights Reserved.
 
 #include "GameplayTagQueryCustomization.h"
-#include "Framework/Application/SlateApplication.h"
 #include "GameplayTagContainer.h"
-#include "Widgets/Text/STextBlock.h"
-#include "Widgets/Input/SButton.h"
-#include "Framework/Docking/TabManager.h"
-
-#include "Editor.h"
 #include "DetailWidgetRow.h"
+#include "IPropertyTypeCustomization.h"
+#include "SGameplayTagQueryEntryBox.h"
 
-#define LOCTEXT_NAMESPACE "GameplayTagQueryCustomization"
-
-void FGameplayTagQueryCustomization::CustomizeHeader(TSharedRef<class IPropertyHandle> InStructPropertyHandle, class FDetailWidgetRow& HeaderRow, IPropertyTypeCustomizationUtils& StructCustomizationUtils)
+void FGameplayTagQueryCustomization::CustomizeHeader(TSharedRef<IPropertyHandle> InStructPropertyHandle, FDetailWidgetRow& HeaderRow, IPropertyTypeCustomizationUtils& StructCustomizationUtils)
 {
-	StructPropertyHandle = InStructPropertyHandle;
-	check(StructPropertyHandle.IsValid());
-
-	PropertyUtilities = StructCustomizationUtils.GetPropertyUtilities();
-	
-	RefreshQueryDescription(); // will call BuildEditableQueryList();
-
-	bool const bReadOnly = StructPropertyHandle->IsEditConst();
-
 	HeaderRow
 		.NameContent()
 		[
-			StructPropertyHandle->CreatePropertyNameWidget()
+			InStructPropertyHandle->CreatePropertyNameWidget()
 		]
 		.ValueContent()
-		.MaxDesiredWidth(512)
+		.VAlign(VAlign_Center)
 		[
-			SNew(SHorizontalBox)
-			+SHorizontalBox::Slot()
-			.AutoWidth()
-			.VAlign(VAlign_Center)
-			[
-				SNew(SVerticalBox)
-				+SVerticalBox::Slot()
-				.AutoHeight()
-				[
-					SNew(SButton)
-					.Text(this, &FGameplayTagQueryCustomization::GetEditButtonText)
-					.OnClicked(this, &FGameplayTagQueryCustomization::OnEditButtonClicked)
-				]
-				+SVerticalBox::Slot()
-				.AutoHeight()
-				[
-					SNew(SButton)
-					.IsEnabled(!bReadOnly)
-					.Text(LOCTEXT("GameplayTagQueryCustomization_Clear", "Clear All"))
-					.OnClicked(this, &FGameplayTagQueryCustomization::OnClearAllButtonClicked)
-					.Visibility(this, &FGameplayTagQueryCustomization::GetClearAllVisibility)
-				]
-			]
-			+SHorizontalBox::Slot()
-			.AutoWidth()
-			[
-				SNew(SBorder)
-				.Padding(4.0f)
-				.Visibility(this, &FGameplayTagQueryCustomization::GetQueryDescVisibility)
-				[
-					SNew(STextBlock)
-					.Text(this, &FGameplayTagQueryCustomization::GetQueryDescText)
-					.ToolTipText(this, &FGameplayTagQueryCustomization::GetQueryDescText)
-					.AutoWrapText(true)
-				]
-			]
+			SNew(SGameplayTagQueryEntryBox)
+			.DescriptionMaxWidth(400.0f)
+			.PropertyHandle(InStructPropertyHandle)
 		];
-
-	GEditor->RegisterForUndo(this);
 }
-
-FText FGameplayTagQueryCustomization::GetQueryDescText() const
-{
-	return FText::FromString(QueryDescription);
-}
-
-FText FGameplayTagQueryCustomization::GetEditButtonText() const
-{
-	if (StructPropertyHandle.IsValid())
-	{
-		bool const bReadOnly = StructPropertyHandle->IsEditConst();
-		return
-			bReadOnly
-			? LOCTEXT("GameplayTagQueryCustomization_View", "View...")
-			: LOCTEXT("GameplayTagQueryCustomization_Edit", "Edit...");
-	}
-
-	return FText();
-}
-
-FReply FGameplayTagQueryCustomization::OnClearAllButtonClicked()
-{
-	if (StructPropertyHandle.IsValid())
-	{
-		FString EmptyQueryAsString;
-    	FGameplayTagQuery::StaticStruct()->ExportText(EmptyQueryAsString, &FGameplayTagQuery::EmptyQuery, &FGameplayTagQuery::EmptyQuery, /*OwnerObject*/nullptr, PPF_None, /*ExportRootScope*/nullptr);
-    	StructPropertyHandle->SetValueFromFormattedString(EmptyQueryAsString);
-	}
-
-	RefreshQueryDescription();
-
-	return FReply::Handled();
-}
-
-EVisibility FGameplayTagQueryCustomization::GetClearAllVisibility() const
-{
-	bool bAtLeastOneQueryIsNonEmpty = false;
-
-	for (auto& EQ : EditableQueries)
-	{
-		if (EQ.TagQuery && EQ.TagQuery->IsEmpty() == false)
-		{
-			bAtLeastOneQueryIsNonEmpty = true;
-			break;
-		}
-	}
-
-	return bAtLeastOneQueryIsNonEmpty ? EVisibility::Visible : EVisibility::Collapsed;
-}
-
-EVisibility FGameplayTagQueryCustomization::GetQueryDescVisibility() const
-{
-	return QueryDescription.IsEmpty() ? EVisibility::Collapsed : EVisibility::Visible;
-}
-
-void FGameplayTagQueryCustomization::RefreshQueryDescription()
-{
-	// Rebuild Editable Containers as container references can become unsafe
-	BuildEditableQueryList();
-
-	// Clear the list
-	QueryDescription.Empty();
-
-	if (EditableQueries.Num() > 1)
-	{
-		QueryDescription = TEXT("Multiple Selected");
-	}
-	else if ((EditableQueries.Num() == 1) && (EditableQueries[0].TagQuery != nullptr))
-	{
-		QueryDescription = EditableQueries[0].TagQuery->GetDescription();
-	}
-}
-
-
-FReply FGameplayTagQueryCustomization::OnEditButtonClicked()
-{
-	if (GameplayTagQueryWidgetWindow.IsValid())
-	{
-		// already open, just show it
-		GameplayTagQueryWidgetWindow->BringToFront(true);
-	}
-	else
-	{
-		if (StructPropertyHandle.IsValid())
-		{
-			TArray<UObject*> OuterObjects;
-			StructPropertyHandle->GetOuterObjects(OuterObjects);
-
-			bool bReadOnly = StructPropertyHandle->IsEditConst();
-
-			FText Title;
-			if (OuterObjects.Num() > 1)
-			{
-				FText const AssetName = FText::Format(LOCTEXT("GameplayTagDetailsBase_MultipleAssets", "{0} Assets"), FText::AsNumber(OuterObjects.Num()));
-				FText const PropertyName = StructPropertyHandle->GetPropertyDisplayName();
-				Title = FText::Format(LOCTEXT("GameplayTagQueryCustomization_BaseWidgetTitle", "Tag Editor: {0} {1}"), PropertyName, AssetName);
-			}
-			else if (OuterObjects.Num() > 0 && OuterObjects[0])
-			{
-				FText const AssetName = FText::FromString(OuterObjects[0]->GetName());
-				FText const PropertyName = StructPropertyHandle->GetPropertyDisplayName();
-				Title = FText::Format(LOCTEXT("GameplayTagQueryCustomization_BaseWidgetTitle", "Tag Editor: {0} {1}"), PropertyName, AssetName);
-			}
-
-			GameplayTagQueryWidgetWindow = SNew(SWindow)
-				.Title(Title)
-				.HasCloseButton(false)
-				.ClientSize(FVector2D(600, 400))
-				[
-					SNew(SGameplayTagQueryWidget, EditableQueries, StructPropertyHandle)
-					.OnClosePreSave(this, &FGameplayTagQueryCustomization::PreSave)
-					.OnSaveAndClose(this, &FGameplayTagQueryCustomization::CloseWidgetWindow, false)
-					.OnCancel(this, &FGameplayTagQueryCustomization::CloseWidgetWindow, true)
-					.ReadOnly(bReadOnly)
-				];
-
-			// NOTE: FGlobalTabmanager::Get()-> is actually dereferencing a SharedReference, not a SharedPtr, so it cannot be null.
-			if (FGlobalTabmanager::Get()->GetRootWindow().IsValid())
-			{
-				FSlateApplication::Get().AddWindowAsNativeChild(GameplayTagQueryWidgetWindow.ToSharedRef(), FGlobalTabmanager::Get()->GetRootWindow().ToSharedRef());
-			}
-			else
-			{
-				FSlateApplication::Get().AddWindow(GameplayTagQueryWidgetWindow.ToSharedRef());
-			}
-		}
-	}
-
-	return FReply::Handled();
-}
-
-FGameplayTagQueryCustomization::~FGameplayTagQueryCustomization()
-{
-	if (GameplayTagQueryWidgetWindow.IsValid())
-	{
-		GameplayTagQueryWidgetWindow->RequestDestroyWindow();
-	}
-	GEditor->UnregisterForUndo(this);
-}
-
-void FGameplayTagQueryCustomization::BuildEditableQueryList()
-{	
-	EditableQueries.Empty();
-
-	if (StructPropertyHandle.IsValid())
-	{
-		TArray<void*> RawStructData;
-		StructPropertyHandle->AccessRawData(RawStructData);
-
-		TArray<UObject*> OuterObjects;
-		StructPropertyHandle->GetOuterObjects(OuterObjects);
-		
-		for (int32 Idx = 0; Idx < RawStructData.Num(); ++Idx)
-		{
-			// Null outer objects may mean that we are inside a UDataTable. This is ok though. We can still dirty the data table via FNotify Hook. (see ::CloseWidgetWindow). However undo will not work.
-			UObject* Obj = OuterObjects.IsValidIndex(Idx) ? OuterObjects[Idx] : nullptr;
-
-			EditableQueries.Add(SGameplayTagQueryWidget::FEditableGameplayTagQueryDatum(Obj, (FGameplayTagQuery*)RawStructData[Idx]));
-		}
-	}	
-}
-
-void FGameplayTagQueryCustomization::PreSave()
-{
-	if (StructPropertyHandle.IsValid())
-	{
-		StructPropertyHandle->NotifyPreChange();
-	}
-}
-
-void FGameplayTagQueryCustomization::CloseWidgetWindow(bool WasCancelled)
-{
-	// Notify change.
-	if (!WasCancelled && StructPropertyHandle.IsValid())
-	{
-		StructPropertyHandle->NotifyPostChange(EPropertyChangeType::ValueSet);
-	}
-
-	if (GameplayTagQueryWidgetWindow.IsValid())
-	{
-		GameplayTagQueryWidgetWindow->RequestDestroyWindow();
-		GameplayTagQueryWidgetWindow = nullptr;
-
-		RefreshQueryDescription();
-	}
-}
-
-#undef LOCTEXT_NAMESPACE
diff --git a/Engine/Plugins/Editor/GameplayTagsEditor/Source/GameplayTagsEditor/Private/GameplayTagQueryCustomization.h b/Engine/Plugins/Editor/GameplayTagsEditor/Source/GameplayTagsEditor/Private/GameplayTagQueryCustomization.h
--- a/Engine/Plugins/Editor/GameplayTagsEditor/Source/GameplayTagsEditor/Private/GameplayTagQueryCustomization.h
+++ b/Engine/Plugins/Editor/GameplayTagsEditor/Source/GameplayTagsEditor/Private/GameplayTagQueryCustomization.h
@@ -3,16 +3,15 @@
 #pragma once
 
 #include "IPropertyTypeCustomization.h"
-#include "SGameplayTagQueryWidget.h"
-#include "EditorUndoClient.h"
 
 struct EVisibility;
-
 class IPropertyHandle;
-class SWindow;
+class FDetailWidgetRow;
+class IDetailChildrenBuilder;
+class IPropertyTypeCustomizationUtils;
 
 /** Customization for the gameplay tag query struct */
-class FGameplayTagQueryCustomization : public IPropertyTypeCustomization, public FEditorUndoClient
+class FGameplayTagQueryCustomization : public IPropertyTypeCustomization
 {
 public:
 	static TSharedRef<IPropertyTypeCustomization> MakeInstance()
@@ -20,51 +19,11 @@ public:
 		return MakeShareable(new FGameplayTagQueryCustomization);
 	}
 
-	~FGameplayTagQueryCustomization();
-
 	/** Overridden to show an edit button to launch the gameplay tag editor */
-	virtual void CustomizeHeader(TSharedRef<class IPropertyHandle> StructPropertyHandle, class FDetailWidgetRow& HeaderRow, IPropertyTypeCustomizationUtils& StructCustomizationUtils) override;
+	virtual void CustomizeHeader(TSharedRef<IPropertyHandle> StructPropertyHandle, FDetailWidgetRow& HeaderRow, IPropertyTypeCustomizationUtils& StructCustomizationUtils) override;
 	
-	virtual void CustomizeChildren(TSharedRef<IPropertyHandle> InStructPropertyHandle, class IDetailChildrenBuilder& ChildBuilder, IPropertyTypeCustomizationUtils& StructCustomizationUtils) override {}
+	virtual void CustomizeChildren(TSharedRef<IPropertyHandle> InStructPropertyHandle, IDetailChildrenBuilder& ChildBuilder, IPropertyTypeCustomizationUtils& StructCustomizationUtils) override {}
 
 private:
-	/** Called when the edit button is clicked; Launches the gameplay tag editor */
-	FReply OnEditButtonClicked();
-
-	/** Overridden to do nothing */
-	FText GetEditButtonText() const;
-
-	/** Called when the "clear all" button is clicked */
-	FReply OnClearAllButtonClicked();
-
-	/** Returns the visibility of the "clear all" button (collapsed when there are no tags) */
-	EVisibility GetClearAllVisibility() const;
-
-	/** Returns the visibility of the tags list box (collapsed when there are no tags) */
-	EVisibility GetQueryDescVisibility() const;
-
-	void RefreshQueryDescription();
-
-	FText GetQueryDescText() const;
-
-	void PreSave();
-
-	void CloseWidgetWindow(bool WasCancelled);
-
-	/** Build List of Editable Queries */
-	void BuildEditableQueryList();
-
-	/** Cached property handle */
-	TSharedPtr<IPropertyHandle> StructPropertyHandle;
-
-	/** The array of queries this objects has */
-	TArray<SGameplayTagQueryWidget::FEditableGameplayTagQueryDatum> EditableQueries;
-
-	/** The Window for the GameplayTagWidget */
-	TSharedPtr<SWindow> GameplayTagQueryWidgetWindow;
-
-	FString QueryDescription;
-
-	TSharedPtr<IPropertyUtilities> PropertyUtilities;
 };
 
diff --git a/Engine/Plugins/Editor/GameplayTagsEditor/Source/GameplayTagsEditor/Private/GameplayTagSearchFilter.cpp b/Engine/Plugins/Editor/GameplayTagsEditor/Source/GameplayTagsEditor/Private/GameplayTagSearchFilter.cpp
--- a/Engine/Plugins/Editor/GameplayTagsEditor/Source/GameplayTagsEditor/Private/GameplayTagSearchFilter.cpp
+++ b/Engine/Plugins/Editor/GameplayTagsEditor/Source/GameplayTagsEditor/Private/GameplayTagSearchFilter.cpp
@@ -8,7 +8,7 @@
 #include "Framework/MultiBox/MultiBoxBuilder.h"
 
 #include "GameplayTagsManager.h"
-#include "SGameplayTagWidget.h"
+#include "SGameplayTagPicker.h"
 
 #include UE_INLINE_GENERATED_CPP_BY_NAME(GameplayTagSearchFilter)
 
@@ -25,9 +25,6 @@ public:
 	FFrontendFilter_GameplayTags(TSharedPtr<FFrontendFilterCategory> InCategory)
 		: FFrontendFilter(InCategory)
 	{
-		TagContainer = MakeShareable(new FGameplayTagContainer);
-
-		EditableContainers.Add(SGameplayTagWidget::FEditableGameplayTagContainerDatum(/*TagContainerOwner=*/ nullptr, TagContainer.Get()));
 	}
 
 	// FFrontendFilter implementation
@@ -46,17 +43,13 @@ public:
 
 protected:
 	// Container of selected search tags (the asset is shown if *any* of these match)
-	TSharedPtr<FGameplayTagContainer> TagContainer;
-
-	// Adaptor for the SGameplayTagWidget to edit our tag container
-	TArray<SGameplayTagWidget::FEditableGameplayTagContainerDatum> EditableContainers;
+	FGameplayTagContainer TagContainer;
 
-protected:
 	bool ProcessStruct(void* Data, UStruct* Struct) const;
 
 	bool ProcessProperty(void* Data, FProperty* Prop) const;
 
-	void OnTagWidgetChanged();
+	void OnTagWidgetChanged(const TArray<FGameplayTagContainer>& TagContainers);
 };
 
 void FFrontendFilter_GameplayTags::ModifyContextMenu(FMenuBuilder& MenuBuilder)
@@ -65,13 +58,17 @@ void FFrontendFilter_GameplayTags::ModifyContextMenu(FMenuBuilder& MenuBuilder)
 
 	MenuBuilder.BeginSection(TEXT("ComparsionSection"), LOCTEXT("ComparisonSectionHeading", "Gameplay Tag(s) to search for"));
 
+	TArray<FGameplayTagContainer> EditableContainers;
+	EditableContainers.Add(TagContainer);
+
 	TSharedRef<SWidget> TagWidget =
 		SNew(SVerticalBox)
 		+SVerticalBox::Slot()
 		.AutoHeight()
 		.MaxHeight(300)
 		[
-			SNew(SGameplayTagWidget, EditableContainers)
+			SNew(SGameplayTagPicker)
+			.TagContainers(EditableContainers)
 			.MultiSelect(true)
 			.OnTagChanged_Raw(this, &FFrontendFilter_GameplayTags::OnTagWidgetChanged)
 		];
@@ -80,7 +77,7 @@ void FFrontendFilter_GameplayTags::ModifyContextMenu(FMenuBuilder& MenuBuilder)
 
 FText FFrontendFilter_GameplayTags::GetDisplayName() const
 {
-	if (TagContainer->Num() == 0)
+	if (TagContainer.Num() == 0)
 	{
 		return LOCTEXT("AnyGameplayTagDisplayName", "Gameplay Tags");
 	}
@@ -89,7 +86,7 @@ FText FFrontendFilter_GameplayTags::GetDisplayName() const
 		FString QueryString;
 
 		int32 Count = 0;
-		for (const FGameplayTag& Tag : *TagContainer.Get())
+		for (const FGameplayTag& Tag : TagContainer)
 		{
 			if (Count > 0)
 			{
@@ -107,7 +104,7 @@ FText FFrontendFilter_GameplayTags::GetDisplayName() const
 
 FText FFrontendFilter_GameplayTags::GetToolTipText() const
 {
-	if (TagContainer->Num() == 0)
+	if (TagContainer.Num() == 0)
 	{
 		return LOCTEXT("AnyGameplayTagFilterDisplayTooltip", "Search for any *loaded* Blueprint or asset that contains a gameplay tag (right-click to choose tags).");
 	}
@@ -120,8 +117,8 @@ FText FFrontendFilter_GameplayTags::GetToolTipText() const
 void FFrontendFilter_GameplayTags::SaveSettings(const FString& IniFilename, const FString& IniSection, const FString& SettingsString) const
 {
 	TArray<FString> TagStrings;
-	TagStrings.Reserve(TagContainer->Num());
-	for (const FGameplayTag& Tag : *TagContainer.Get())
+	TagStrings.Reserve(TagContainer.Num());
+	for (const FGameplayTag& Tag : TagContainer)
 	{
 		TagStrings.Add(Tag.GetTagName().ToString());
 	}
@@ -136,19 +133,24 @@ void FFrontendFilter_GameplayTags::LoadSettings(const FString& IniFilename, cons
 	TArray<FString> TagStrings;
 	GConfig->GetArray(*IniSection, *(SettingsString + TEXT(".Tags")), /*out*/ TagStrings, IniFilename);
 
-	TagContainer->Reset();
+	TagContainer.Reset();
 	for (const FString& TagString : TagStrings)
 	{
 		FGameplayTag NewTag = Manager.RequestGameplayTag(*TagString, /*bErrorIfNotFound=*/ false);
 		if (NewTag.IsValid())
 		{
-			TagContainer->AddTag(NewTag);
+			TagContainer.AddTag(NewTag);
 		}
 	}
 }
 
-void FFrontendFilter_GameplayTags::OnTagWidgetChanged()
+void FFrontendFilter_GameplayTags::OnTagWidgetChanged(const TArray<FGameplayTagContainer>& TagContainers)
 {
+	if (TagContainers.Num() > 0)
+	{
+		TagContainer = TagContainers[0];
+	}
+	
 	BroadcastChangedEvent();
 }
 
@@ -177,8 +179,8 @@ bool FFrontendFilter_GameplayTags::ProcessProperty(void* Data, FProperty* Prop)
 		{
 			FGameplayTag& ThisTag = *static_cast<FGameplayTag*>(InnerData);
 
-			const bool bAnyTagIsOK = TagContainer->Num() == 0;
-			const bool bPassesTagSearch = bAnyTagIsOK || ThisTag.MatchesAny(*TagContainer);
+			const bool bAnyTagIsOK = TagContainer.Num() == 0;
+			const bool bPassesTagSearch = bAnyTagIsOK || ThisTag.MatchesAny(TagContainer);
 
 			return bPassesTagSearch;
 		}
diff --git a/Engine/Plugins/Editor/GameplayTagsEditor/Source/GameplayTagsEditor/Private/GameplayTagsEditorModule.cpp b/Engine/Plugins/Editor/GameplayTagsEditor/Source/GameplayTagsEditor/Private/GameplayTagsEditorModule.cpp
--- a/Engine/Plugins/Editor/GameplayTagsEditor/Source/GameplayTagsEditor/Private/GameplayTagsEditorModule.cpp
+++ b/Engine/Plugins/Editor/GameplayTagsEditor/Source/GameplayTagsEditor/Private/GameplayTagsEditorModule.cpp
@@ -27,6 +27,8 @@
 #include "Stats/StatsMisc.h"
 #include "Subsystems/ImportSubsystem.h"
 #include "UObject/ObjectSaveContext.h"
+#include "GameplayTagStyle.h"
+#include "SGameplayTagPicker.h"
 
 #define LOCTEXT_NAMESPACE "GameplayTagEditor"
 
@@ -41,6 +43,7 @@ public:
 	virtual void StartupModule() override
 	{
 		FCoreDelegates::OnPostEngineInit.AddRaw(this, &FGameplayTagsEditorModule::OnPostEngineInit);
+		FGameplayTagStyle::Initialize();
 	}
 	
 	void OnPostEngineInit()
@@ -830,16 +833,20 @@ public:
 			return SNullWidget::NullWidget;
 		}
 
-		TArray<SGameplayTagWidget::FEditableGameplayTagContainerDatum> EditableContainers;
-		EditableContainers.Emplace(nullptr, GameplayTagContainer.Get());
+		TArray<FGameplayTagContainer> EditableContainers;
+		EditableContainers.Emplace(*GameplayTagContainer);
 
-		// This will keep the shared ptr and delegate alive as long as the widget is
-		SGameplayTagWidget::FOnTagChanged OnChanged = SGameplayTagWidget::FOnTagChanged::CreateLambda([OnSetTag, GameplayTagContainer]()
+		SGameplayTagPicker::FOnTagChanged OnChanged = SGameplayTagPicker::FOnTagChanged::CreateLambda([OnSetTag, GameplayTagContainer](const TArray<FGameplayTagContainer>& TagContainers)
 		{
-			OnSetTag.Execute(*GameplayTagContainer.Get());
+			if (TagContainers.Num() > 0)
+			{
+				*GameplayTagContainer.Get() = TagContainers[0];
+				OnSetTag.Execute(*GameplayTagContainer.Get());
+			}
 		});
 
-		return SNew(SGameplayTagWidget, EditableContainers)
+		return SNew(SGameplayTagPicker)
+			.TagContainers(EditableContainers)
 			.Filter(FilterString)
 			.ReadOnly(false)
 			.MultiSelect(true)
@@ -853,20 +860,20 @@ public:
 			return SNullWidget::NullWidget;
 		}
 
-		// Make a wrapper tag container
-		TSharedPtr<FGameplayTagContainer> GameplayTagContainer = MakeShareable(new FGameplayTagContainer(*GameplayTag.Get()));
-
-		TArray<SGameplayTagWidget::FEditableGameplayTagContainerDatum> EditableContainers;
-		EditableContainers.Emplace(nullptr, GameplayTagContainer.Get());
+		TArray<FGameplayTagContainer> EditableContainers;
+		EditableContainers.Emplace(*GameplayTag);
 
-		// This will keep the shared ptr and delegate alive as long as the widget is
-		SGameplayTagWidget::FOnTagChanged OnChanged = SGameplayTagWidget::FOnTagChanged::CreateLambda([OnSetTag, GameplayTag, GameplayTagContainer]()
+		SGameplayTagPicker::FOnTagChanged OnChanged = SGameplayTagPicker::FOnTagChanged::CreateLambda([OnSetTag, GameplayTag](const TArray<FGameplayTagContainer>& TagContainers)
 		{
-			*GameplayTag.Get() = GameplayTagContainer->First();
-			OnSetTag.Execute(*GameplayTag.Get());
+			if (TagContainers.Num() > 0)
+			{
+				*GameplayTag.Get() = TagContainers[0].First();
+				OnSetTag.Execute(*GameplayTag.Get());
+			}
 		});
 
-		return SNew(SGameplayTagWidget, EditableContainers)
+		return SNew(SGameplayTagPicker)
+			.TagContainers(EditableContainers)
 			.Filter(FilterString)
 			.ReadOnly(false)
 			.MultiSelect(false)
diff --git a/Engine/Plugins/Editor/GameplayTagsEditor/Source/GameplayTagsEditor/Private/GameplayTagsGraphPanelPinFactory.h b/Engine/Plugins/Editor/GameplayTagsEditor/Source/GameplayTagsEditor/Private/GameplayTagsGraphPanelPinFactory.h
--- a/Engine/Plugins/Editor/GameplayTagsEditor/Source/GameplayTagsEditor/Private/GameplayTagsGraphPanelPinFactory.h
+++ b/Engine/Plugins/Editor/GameplayTagsEditor/Source/GameplayTagsEditor/Private/GameplayTagsGraphPanelPinFactory.h
@@ -5,6 +5,7 @@
 #include "EdGraphUtilities.h"
 #include "GameplayTagContainer.h"
 #include "EdGraphSchema_K2.h"
+#include "SGameplayTagGraphPin.h"
 #include "SGameplayTagContainerGraphPin.h"
 #include "SGameplayTagQueryGraphPin.h"
 
diff --git a/Engine/Plugins/Editor/GameplayTagsEditor/Source/GameplayTagsEditor/Private/GameplayTagsSettingsCustomization.cpp b/Engine/Plugins/Editor/GameplayTagsEditor/Source/GameplayTagsEditor/Private/GameplayTagsSettingsCustomization.cpp
--- a/Engine/Plugins/Editor/GameplayTagsEditor/Source/GameplayTagsEditor/Private/GameplayTagsSettingsCustomization.cpp
+++ b/Engine/Plugins/Editor/GameplayTagsEditor/Source/GameplayTagsEditor/Private/GameplayTagsSettingsCustomization.cpp
@@ -5,7 +5,10 @@
 #include "GameplayTagsModule.h"
 #include "DetailLayoutBuilder.h"
 #include "DetailWidgetRow.h"
-#include "SGameplayTagWidget.h"
+#include "SGameplayTagPicker.h"
+#include "SAddNewGameplayTagSourceWidget.h"
+#include "Widgets/Input/SButton.h"
+#include "Editor.h"
 
 #define LOCTEXT_NAMESPACE "FGameplayTagsSettingsCustomization"
 
@@ -16,7 +19,6 @@ TSharedRef<IDetailCustomization> FGameplayTagsSettingsCustomization::MakeInstanc
 
 FGameplayTagsSettingsCustomization::FGameplayTagsSettingsCustomization()
 {
-	IGameplayTagsModule::OnTagSettingsChanged.AddRaw(this, &FGameplayTagsSettingsCustomization::OnTagTreeChanged);
 }
 
 FGameplayTagsSettingsCustomization::~FGameplayTagsSettingsCustomization()
@@ -26,9 +28,6 @@ FGameplayTagsSettingsCustomization::~FGameplayTagsSettingsCustomization()
 
 void FGameplayTagsSettingsCustomization::CustomizeDetails(IDetailLayoutBuilder& DetailLayout)
 {
-	const float MaxPropertyWidth = 480.0f;
-	const float MaxPropertyHeight = 240.0f;
-
 	IDetailCategoryBuilder& GameplayTagsCategory = DetailLayout.EditCategory("GameplayTags");
 	{
 		TArray<TSharedRef<IPropertyHandle>> GameplayTagsProperties;
@@ -37,33 +36,105 @@ void FGameplayTagsSettingsCustomization::CustomizeDetails(IDetailLayoutBuilder&
 		TSharedPtr<IPropertyHandle> TagListProperty = DetailLayout.GetProperty(GET_MEMBER_NAME_CHECKED(UGameplayTagsList, GameplayTagList), UGameplayTagsList::StaticClass());
 		TagListProperty->MarkHiddenByCustomization();
 
+		TSharedPtr<IPropertyHandle> NewTagSourceProperty = DetailLayout.GetProperty(GET_MEMBER_NAME_CHECKED(UGameplayTagsSettings, NewTagSource));
+		NewTagSourceProperty->MarkHiddenByCustomization();
+
 		for (TSharedPtr<IPropertyHandle> Property : GameplayTagsProperties)
 		{
-			if (Property->GetProperty() != TagListProperty->GetProperty())
+			if (Property->GetProperty() == TagListProperty->GetProperty())
 			{
-				GameplayTagsCategory.AddProperty(Property);
+				// Button to open tag manager
+				GameplayTagsCategory.AddCustomRow(TagListProperty->GetPropertyDisplayName(), /*bForAdvanced*/false)
+				.NameContent()
+				[
+					TagListProperty->CreatePropertyNameWidget()
+				]
+				.ValueContent()
+				[
+					SNew(SButton)
+					.VAlign(VAlign_Center)
+					.HAlign(HAlign_Center)
+					.OnClicked_Lambda([this]()
+					{
+						FGameplayTagManagerWindowArgs Args;
+						Args.bRestrictedTags = false;
+						UE::GameplayTags::Editor::OpenGameplayTagManager(Args);
+						return FReply::Handled();
+					})
+					[
+						SNew(SHorizontalBox)
+						+SHorizontalBox::Slot()
+						.AutoWidth()
+						.Padding(FMargin(0,0,4,0))
+						[
+							SNew( SImage )
+							.Image(FAppStyle::GetBrush("Icons.Settings"))
+							.ColorAndOpacity(FSlateColor::UseForeground())
+						]
+						+SHorizontalBox::Slot()
+						.AutoWidth()
+						[
+							SNew(STextBlock)
+							.Text(LOCTEXT("ManageGameplayTags", "Manage Gameplay Tags..."))
+						]
+					]
+				];
 			}
-			else
+			else if (Property->GetProperty() ==  NewTagSourceProperty->GetProperty())
 			{
-				// Create a custom widget for the tag list
-
-				GameplayTagsCategory.AddCustomRow(TagListProperty->GetPropertyDisplayName(), false)
+				// Button to open add source dialog
+				GameplayTagsCategory.AddCustomRow(NewTagSourceProperty->GetPropertyDisplayName(), /*bForAdvanced*/false)
 				.NameContent()
 				[
-					TagListProperty->CreatePropertyNameWidget()
+					NewTagSourceProperty->CreatePropertyNameWidget()
 				]
 				.ValueContent()
-				.MaxDesiredWidth(MaxPropertyWidth)
 				[
-					SAssignNew(TagWidget, SGameplayTagWidget, TArray<SGameplayTagWidget::FEditableGameplayTagContainerDatum>())
-					.Filter(TEXT(""))
-					.MultiSelect(false)
-					.GameplayTagUIMode(EGameplayTagUIMode::ManagementMode)
-					.MaxHeight(MaxPropertyHeight)
-					.OnTagChanged(this, &FGameplayTagsSettingsCustomization::OnTagChanged)
-					.RestrictedTags(false)
+					SNew(SButton)
+					.VAlign(VAlign_Center)
+					.HAlign(HAlign_Center)
+					.OnClicked_Lambda([this]()
+					{
+						const FVector2D WindowSize(320, 110);
+						const TSharedRef<SWindow> Window = SNew(SWindow)
+							.Title(LOCTEXT("AddNewGameplayTagSourceTitle", "Add new Gameplay Tag Source"))
+							.ClientSize(WindowSize)
+							.SizingRule(ESizingRule::Autosized)
+							.AutoCenter(EAutoCenter::PreferredWorkArea)
+							.SupportsMaximize(false)
+							.SupportsMinimize(false)
+							.Content()
+							[
+								SNew(SAddNewGameplayTagSourceWidget)
+							];
+
+						GEditor->EditorAddModalWindow(Window);
+
+						return FReply::Handled();
+					})
+					[
+						SNew(SHorizontalBox)
+						+SHorizontalBox::Slot()
+						.AutoWidth()
+						.Padding(FMargin(0,0,4,0))
+						[
+							SNew( SImage )
+							.Image(FAppStyle::GetBrush("Icons.Plus"))
+							.ColorAndOpacity(FSlateColor::UseForeground())
+						]
+						+SHorizontalBox::Slot()
+						.AutoWidth()
+						[
+							SNew(STextBlock)
+							.Text(LOCTEXT("AddNewGameplayTagSource", "Add new Gameplay Tag source..."))
+						]
+					]
 				];
 			}
+			else
+			{
+				GameplayTagsCategory.AddProperty(Property);
+			}
 		}
 	}
 
@@ -79,23 +150,41 @@ void FGameplayTagsSettingsCustomization::CustomizeDetails(IDetailLayoutBuilder&
 		{
 			if (Property->GetProperty() == RestrictedTagListProperty->GetProperty())
 			{
-				// Create a custom widget for the restricted tag list
-
+				// Button to open restricted tag manager
 				AdvancedGameplayTagsCategory.AddCustomRow(RestrictedTagListProperty->GetPropertyDisplayName(), true)
 				.NameContent()
 				[
 					RestrictedTagListProperty->CreatePropertyNameWidget()
 				]
 				.ValueContent()
-				.MaxDesiredWidth(MaxPropertyWidth)
 				[
-					SAssignNew(RestrictedTagWidget, SGameplayTagWidget, TArray<SGameplayTagWidget::FEditableGameplayTagContainerDatum>())
-					.Filter(TEXT(""))
-					.MultiSelect(false)
-					.GameplayTagUIMode(EGameplayTagUIMode::ManagementMode)
-					.MaxHeight(MaxPropertyHeight)
-					.OnTagChanged(this, &FGameplayTagsSettingsCustomization::OnTagChanged)
-					.RestrictedTags(true)
+					SNew(SButton)
+					.VAlign(VAlign_Center)
+					.HAlign(HAlign_Center)
+					.OnClicked_Lambda([this]()
+					{
+						FGameplayTagManagerWindowArgs Args;
+						Args.bRestrictedTags = true;
+						UE::GameplayTags::Editor::OpenGameplayTagManager(Args);
+						return FReply::Handled();
+					})
+					[
+						SNew(SHorizontalBox)
+						+SHorizontalBox::Slot()
+						.AutoWidth()
+						.Padding(FMargin(0,0,4,0))
+						[
+							SNew( SImage )
+							.Image(FAppStyle::GetBrush("Icons.Settings"))
+							.ColorAndOpacity(FSlateColor::UseForeground())
+						]
+						+SHorizontalBox::Slot()
+						.AutoWidth()
+						[
+							SNew(STextBlock)
+							.Text(LOCTEXT("ManageRestrictedGameplayTags", "Manage Restricted Gameplay Tags..."))
+						]
+					]
 				];
 			}
 			else
@@ -106,30 +195,4 @@ void FGameplayTagsSettingsCustomization::CustomizeDetails(IDetailLayoutBuilder&
 	}
 }
 
-void FGameplayTagsSettingsCustomization::OnTagChanged()
-{
-	if (TagWidget.IsValid())
-	{
-		TagWidget->RefreshTags();
-	}
-
-	if (RestrictedTagWidget.IsValid())
-	{
-		RestrictedTagWidget->RefreshTags();
-	}
-}
-
-void FGameplayTagsSettingsCustomization::OnTagTreeChanged()
-{
-	if (TagWidget.IsValid())
-	{
-		TagWidget->RefreshOnNextTick();
-	}
-
-	if (RestrictedTagWidget.IsValid())
-	{
-		RestrictedTagWidget->RefreshOnNextTick();
-	}
-}
-
 #undef LOCTEXT_NAMESPACE
diff --git a/Engine/Plugins/Editor/GameplayTagsEditor/Source/GameplayTagsEditor/Private/GameplayTagsSettingsCustomization.h b/Engine/Plugins/Editor/GameplayTagsEditor/Source/GameplayTagsEditor/Private/GameplayTagsSettingsCustomization.h
--- a/Engine/Plugins/Editor/GameplayTagsEditor/Source/GameplayTagsEditor/Private/GameplayTagsSettingsCustomization.h
+++ b/Engine/Plugins/Editor/GameplayTagsEditor/Source/GameplayTagsEditor/Private/GameplayTagsSettingsCustomization.h
@@ -5,7 +5,6 @@
 #include "IDetailCustomization.h"
 
 class SGameplayTagWidget;
-
 class IDetailLayoutBuilder;
 
 //////////////////////////////////////////////////////////////////////////
@@ -15,7 +14,7 @@ class FGameplayTagsSettingsCustomization : public IDetailCustomization
 {
 public:
 	FGameplayTagsSettingsCustomization();
-	virtual ~FGameplayTagsSettingsCustomization();
+	virtual ~FGameplayTagsSettingsCustomization() override;
 
 	/** Makes a new instance of this detail layout class for a specific detail view requesting it */
 	static TSharedRef<IDetailCustomization> MakeInstance();
@@ -23,16 +22,4 @@ public:
 	// IDetailCustomization interface
 	virtual void CustomizeDetails(IDetailLayoutBuilder& DetailLayout) override;
 	// End of IDetailCustomization interface
-
-private:
-
-	/** Callback for when a tag changes */
-	void OnTagChanged();
-
-	/** Module callback for when the tag tree changes */
-	void OnTagTreeChanged();
-
-	TSharedPtr<SGameplayTagWidget> TagWidget;
-
-	TSharedPtr<SGameplayTagWidget> RestrictedTagWidget;
 };
diff --git a/Engine/Plugins/Editor/GameplayTagsEditor/Source/GameplayTagsEditor/Private/SAddNewGameplayTagSourceWidget.cpp b/Engine/Plugins/Editor/GameplayTagsEditor/Source/GameplayTagsEditor/Private/SAddNewGameplayTagSourceWidget.cpp
--- a/Engine/Plugins/Editor/GameplayTagsEditor/Source/GameplayTagsEditor/Private/SAddNewGameplayTagSourceWidget.cpp
+++ b/Engine/Plugins/Editor/GameplayTagsEditor/Source/GameplayTagsEditor/Private/SAddNewGameplayTagSourceWidget.cpp
@@ -9,6 +9,7 @@
 #include "Misc/Paths.h"
 #include "Widgets/Input/SComboBox.h"
 #include "Widgets/Input/SEditableTextBox.h"
+#include "Widgets/Layout/SGridPanel.h"
 
 #define LOCTEXT_NAMESPACE "AddNewGameplayTagSourceWidget"
 
@@ -28,76 +29,71 @@ void SAddNewGameplayTagSourceWidget::Construct(const FArguments& InArgs)
 
 	ChildSlot
 	[
-		SNew(SVerticalBox)
-
-		// Tag Name
-		+ SVerticalBox::Slot()
-		.AutoHeight()
-		.VAlign(VAlign_Top)
+		SNew(SBorder)
+		.BorderImage(FAppStyle::GetBrush("ToolPanel.GroupBorder"))
 		[
-			SNew(SHorizontalBox)
-			+ SHorizontalBox::Slot()
-			.Padding(2.0f, 4.0f)
-			.AutoWidth()
+			SNew(SGridPanel)
+			// Tag Name
+			+ SGridPanel::Slot(0, 0)
+			.Padding(5)
+			.VAlign(VAlign_Center)
+			.HAlign(HAlign_Left)
 			[
-				SNew(STextBlock)
-				.Text(LOCTEXT("NewSourceName", "Name:"))
+				SNew(SBox)
+				.MinDesiredWidth(150.0f)
+				[
+					SNew(STextBlock)
+					.Font(FAppStyle::GetFontStyle( TEXT("PropertyWindow.NormalFont")))
+					.Text(LOCTEXT("NewSourceName", "Name:"))
+				]
 			]
-
-			+ SHorizontalBox::Slot()
-			.Padding(2.0f, 2.0f)
-			.FillWidth(1.0f)
-			.HAlign(HAlign_Right)
+			+ SGridPanel::Slot(1, 0)
+			.Padding(5)
+			.VAlign(VAlign_Center)
+			.HAlign(HAlign_Left)
 			[
-				SAssignNew(SourceNameTextBox, SEditableTextBox)
-				.MinDesiredWidth(240.0f)
-				.HintText(HintText)
+				SNew(SBox)
+				.MinDesiredWidth(300.0f)
+				[
+					SAssignNew(SourceNameTextBox, SEditableTextBox)
+					.HintText(HintText)
+				]
 			]
-		]
 
-		// Tag source root
-		+ SVerticalBox::Slot()
-		.AutoHeight()
-		.VAlign(VAlign_Top)
-		[
-			SNew(SHorizontalBox)
-			+ SHorizontalBox::Slot()
-			.Padding(2.0f, 6.0f)
-			.AutoWidth()
+			// Tag source root
+			+ SGridPanel::Slot(0, 1)
+			.Padding(5)
+			.VAlign(VAlign_Center)
+			.HAlign(HAlign_Left)
 			[
 				SNew(STextBlock)
 				.Text(LOCTEXT("ConfigPathLabel", "Config Path:"))
+				.Font(FAppStyle::GetFontStyle( TEXT("PropertyWindow.NormalFont")))
 				.ToolTipText(LOCTEXT("RootPathTooltip", "Set the base config path for added source, this includes paths from plugins and other places that call AddTagIniSearchPath"))
 			]
-
-			+ SHorizontalBox::Slot()
-			.Padding(2.0f, 2.0f)
-			.FillWidth(1.0f)
-			.HAlign(HAlign_Right)
+			+ SGridPanel::Slot(1, 1)
+			.Padding(5)
+			.HAlign(HAlign_Left)
 			[
-				SAssignNew(TagRootsComboBox, SComboBox<TSharedPtr<FString> >)
-				.OptionsSource(&TagRoots)
-				.OnGenerateWidget(this, &SAddNewGameplayTagSourceWidget::OnGenerateTagRootsComboBox)
-				.ContentPadding(2.0f)
-				.Content()
+				SNew(SBox)
+				.MinDesiredWidth(300.0f)
 				[
-					SNew(STextBlock)
-					.Text(this, &SAddNewGameplayTagSourceWidget::CreateTagRootsComboBoxContent)
-					.Font(IDetailLayoutBuilder::GetDetailFont())
+					SAssignNew(TagRootsComboBox, SComboBox<TSharedPtr<FString> >)
+					.OptionsSource(&TagRoots)
+					.OnGenerateWidget(this, &SAddNewGameplayTagSourceWidget::OnGenerateTagRootsComboBox)
+					.Content()
+					[
+						SNew(STextBlock)
+						.Text(this, &SAddNewGameplayTagSourceWidget::CreateTagRootsComboBoxContent)
+						.Font(IDetailLayoutBuilder::GetDetailFont())
+					]
 				]
 			]
-		]
-
-		// Add Source Button
-		+SVerticalBox::Slot()
-		.AutoHeight()
-		.VAlign(VAlign_Top)
-		.HAlign(HAlign_Center)
-		.Padding(8.0f)
-		[
-			SNew(SHorizontalBox)
-			+ SHorizontalBox::Slot()
-			.AutoWidth()
+			
+			// Add Source Button
+			+ SGridPanel::Slot(1, 2)
+			.Padding(5)
+			.HAlign(HAlign_Left)
 			[
 				SNew(SButton)
 				.Text(LOCTEXT("AddNew", "Add New Source"))
diff --git a/Engine/Plugins/Editor/GameplayTagsEditor/Source/GameplayTagsEditor/Private/SAddNewGameplayTagWidget.cpp b/Engine/Plugins/Editor/GameplayTagsEditor/Source/GameplayTagsEditor/Private/SAddNewGameplayTagWidget.cpp
--- a/Engine/Plugins/Editor/GameplayTagsEditor/Source/GameplayTagsEditor/Private/SAddNewGameplayTagWidget.cpp
+++ b/Engine/Plugins/Editor/GameplayTagsEditor/Source/GameplayTagsEditor/Private/SAddNewGameplayTagWidget.cpp
@@ -2,7 +2,6 @@
 
 #include "SAddNewGameplayTagWidget.h"
 #include "DetailLayoutBuilder.h"
-#include "Framework/Views/TableViewMetadata.h"
 #include "GameplayTagsEditorModule.h"
 #include "GameplayTagsManager.h"
 #include "GameplayTagsModule.h"
@@ -11,6 +10,9 @@
 #include "Misc/MessageDialog.h"
 #include "Widgets/Input/SComboBox.h"
 #include "Widgets/Input/SEditableTextBox.h"
+#include "Widgets/Layout/SGridPanel.h"
+#include "Widgets/Notifications/SNotificationList.h"
+#include "Framework/Notifications/NotificationManager.h"
 
 #define LOCTEXT_NAMESPACE "AddNewGameplayTagWidget"
 
@@ -24,7 +26,6 @@ SAddNewGameplayTagWidget::~SAddNewGameplayTagWidget()
 
 void SAddNewGameplayTagWidget::Construct(const FArguments& InArgs)
 {
-	
 	FText HintText = LOCTEXT("NewTagNameHint", "X.Y.Z");
 	DefaultNewName = InArgs._NewTagName;
 	if (DefaultNewName.IsEmpty() == false)
@@ -32,7 +33,6 @@ void SAddNewGameplayTagWidget::Construct(const FArguments& InArgs)
 		HintText = FText::FromString(DefaultNewName);
 	}
 
-
 	bAddingNewTag = false;
 	bShouldGetKeyboardFocus = false;
 
@@ -44,119 +44,122 @@ void SAddNewGameplayTagWidget::Construct(const FArguments& InArgs)
 
 	ChildSlot
 	[
-		SNew(SVerticalBox)
-
-		// Tag Name
-		+ SVerticalBox::Slot()
-		.AutoHeight()
-		.VAlign(VAlign_Top)
+		SNew(SBorder)
+		.BorderImage(FAppStyle::GetBrush("ToolPanel.GroupBorder"))
 		[
-			SNew(SHorizontalBox)
-			+ SHorizontalBox::Slot()
-			.Padding(2.0f, 4.0f)
-			.AutoWidth()
+			SNew(SGridPanel)
+			
+			// Tag Name
+			+ SGridPanel::Slot(0, 0)
+			.Padding(5)
+			.VAlign(VAlign_Center)
+			.HAlign(HAlign_Left)
 			[
-				SNew(STextBlock)
-				.Text(LOCTEXT("NewTagName", "Name:"))
+				SNew(SBox)
+				.MinDesiredWidth(150.0f)
+				[
+					SNew(STextBlock)
+					.Font(FAppStyle::GetFontStyle( TEXT("PropertyWindow.NormalFont")))
+					.Text(LOCTEXT("NewTagName", "Name:"))
+				]
 			]
-
-			+ SHorizontalBox::Slot()
-			.Padding(2.0f, 2.0f)
-			.FillWidth(1.0f)
-			.HAlign(HAlign_Right)
+			+ SGridPanel::Slot(1, 0)
+			.Padding(5)
+			.VAlign(VAlign_Center)
+			.HAlign(HAlign_Left)
 			[
-				SAssignNew(TagNameTextBox, SEditableTextBox)
-				.MinDesiredWidth(240.0f)
-				.HintText(HintText)
-				.OnTextCommitted(this, &SAddNewGameplayTagWidget::OnCommitNewTagName)
+				SNew(SBox)
+				.MinDesiredWidth(300.0f)
+				[
+					SAssignNew(TagNameTextBox, SEditableTextBox)
+					.HintText(HintText)
+					.OnTextCommitted(this, &SAddNewGameplayTagWidget::OnCommitNewTagName)
+				]
 			]
-		]
-
-		// Tag Comment
-		+ SVerticalBox::Slot()
-		.AutoHeight()
-		.VAlign(VAlign_Top)
-		[
-			SNew(SHorizontalBox)
-			+ SHorizontalBox::Slot()
-			.Padding(2.0f, 4.0f)
-			.AutoWidth()
+			
+			// Tag Comment
+			+ SGridPanel::Slot(0, 1)
+			.Padding(5)
+			.VAlign(VAlign_Center)
+			.HAlign(HAlign_Left)
 			[
 				SNew(STextBlock)
+				.Font(FAppStyle::GetFontStyle( TEXT("PropertyWindow.NormalFont")))
 				.Text(LOCTEXT("TagComment", "Comment:"))
 			]
-
-			+ SHorizontalBox::Slot()
-			.Padding(2.0f, 2.0f)
-			.FillWidth(1.0f)
-			.HAlign(HAlign_Right)
+			+ SGridPanel::Slot(1, 1)
+			.Padding(5)
+			.VAlign(VAlign_Center)
+			.HAlign(HAlign_Left)
 			[
-				SAssignNew(TagCommentTextBox, SEditableTextBox)
-				.MinDesiredWidth(240.0f)
-				.HintText(LOCTEXT("TagCommentHint", "Comment"))
-				.OnTextCommitted(this, &SAddNewGameplayTagWidget::OnCommitNewTagName)
+				SNew(SBox)
+				.MinDesiredWidth(300.0f)
+				[
+					SAssignNew(TagCommentTextBox, SEditableTextBox)
+					.HintText(LOCTEXT("TagCommentHint", "Comment"))
+					.OnTextCommitted(this, &SAddNewGameplayTagWidget::OnCommitNewTagName)
+				]
 			]
-		]
 
-		// Tag Location
-		+ SVerticalBox::Slot()
-		.AutoHeight()
-		.VAlign(VAlign_Top)
-		[
-			SNew(SHorizontalBox)
-			+ SHorizontalBox::Slot()
-			.Padding(2.0f, 6.0f)
-			.AutoWidth()
+			// Tag Location
+			+ SGridPanel::Slot(0, 2)
+			.Padding(5)
+			.VAlign(VAlign_Center)
+			.HAlign(HAlign_Left)
 			[
 				SNew(STextBlock)
 				.Text(LOCTEXT("CreateTagSource", "Source:"))
+				.Font(IDetailLayoutBuilder::GetDetailFont())
 			]
-
-			+ SHorizontalBox::Slot()
-			.Padding(2.0f, 2.0f)
-			.FillWidth(1.0f)
-			.HAlign(HAlign_Right)
-			[
-				SAssignNew(TagSourcesComboBox, SComboBox<TSharedPtr<FName> >)
-				.OptionsSource(&TagSources)
-				.OnGenerateWidget(this, &SAddNewGameplayTagWidget::OnGenerateTagSourcesComboBox)
-				.ToolTipText(this, &SAddNewGameplayTagWidget::CreateTagSourcesComboBoxToolTip)
-				.ContentPadding(2.0f)
-				.Content()
-				[
-					SNew(STextBlock)
-					.Text(this, &SAddNewGameplayTagWidget::CreateTagSourcesComboBoxContent)
-					.Font(IDetailLayoutBuilder::GetDetailFont())
-				]
-			]
-
-			+ SHorizontalBox::Slot()
-			.AutoWidth()
+			+ SGridPanel::Slot(1, 2)
+			.Padding(5)
 			.VAlign(VAlign_Center)
+			.HAlign(HAlign_Left)
 			[
-				SNew( SButton )
-				.ButtonStyle( FAppStyle::Get(), "NoBorder" )
-				.Visibility(this, &SAddNewGameplayTagWidget::OnGetTagSourceFavoritesVisibility)
-				.OnClicked(this, &SAddNewGameplayTagWidget::OnToggleTagSourceFavoriteClicked)
-				.ToolTipText(LOCTEXT("ToggleFavoriteTooltip", "Toggle whether or not this tag source is your favorite source (new tags will go into your favorite source by default)"))
-				.ContentPadding(0)
+				SNew(SBox)
+				.MinDesiredWidth(300.0f)
 				[
-					SNew(SImage)
-					.Image(this, &SAddNewGameplayTagWidget::OnGetTagSourceFavoriteImage)
+					SNew(SHorizontalBox)
+
+					+ SHorizontalBox::Slot()
+					.Padding(2.0f, 2.0f)
+					.FillWidth(1.0f)
+					.VAlign(VAlign_Center)
+					[
+						SAssignNew(TagSourcesComboBox, SComboBox<TSharedPtr<FName> >)
+						.OptionsSource(&TagSources)
+						.OnGenerateWidget(this, &SAddNewGameplayTagWidget::OnGenerateTagSourcesComboBox)
+						.ToolTipText(this, &SAddNewGameplayTagWidget::CreateTagSourcesComboBoxToolTip)
+						.Content()
+						[
+							SNew(STextBlock)
+							.Text(this, &SAddNewGameplayTagWidget::CreateTagSourcesComboBoxContent)
+							.Font(IDetailLayoutBuilder::GetDetailFont())
+						]
+					]
+
+					+ SHorizontalBox::Slot()
+					.AutoWidth()
+					.VAlign(VAlign_Center)
+					[
+						SNew( SButton )
+						.ButtonStyle( FAppStyle::Get(), "NoBorder" )
+						.Visibility(this, &SAddNewGameplayTagWidget::OnGetTagSourceFavoritesVisibility)
+						.OnClicked(this, &SAddNewGameplayTagWidget::OnToggleTagSourceFavoriteClicked)
+						.ToolTipText(LOCTEXT("ToggleFavoriteTooltip", "Toggle whether or not this tag source is your favorite source (new tags will go into your favorite source by default)"))
+						.ContentPadding(0)
+						[
+							SNew(SImage)
+							.Image(this, &SAddNewGameplayTagWidget::OnGetTagSourceFavoriteImage)
+						]
+					]
 				]
 			]
-		]
 
-		// Add Tag Button
-		+SVerticalBox::Slot()
-		.AutoHeight()
-		.VAlign(VAlign_Top)
-		.HAlign(HAlign_Center)
-		.Padding(8.0f)
-		[
-			SNew(SHorizontalBox)
-			+ SHorizontalBox::Slot()
-			.AutoWidth()
+			// Add Tag Button
+			+ SGridPanel::Slot(1, 3)
+			.Padding(5)
+			.HAlign(HAlign_Left)
 			[
 				SNew(SButton)
 				.Text(LOCTEXT("AddNew", "Add New Tag"))
@@ -175,7 +178,8 @@ EVisibility SAddNewGameplayTagWidget::OnGetTagSourceFavoritesVisibility() const
 
 FReply SAddNewGameplayTagWidget::OnToggleTagSourceFavoriteClicked()
 {
-	const FName ActiveTagSource = *TagSourcesComboBox->GetSelectedItem().Get();
+	const bool bHasSelectedItem = TagSourcesComboBox.IsValid() && TagSourcesComboBox->GetSelectedItem().IsValid();
+	const FName ActiveTagSource = bHasSelectedItem ? *TagSourcesComboBox->GetSelectedItem().Get() : FName();
 	const bool bWasFavorite = FGameplayTagSource::GetFavoriteName() == ActiveTagSource;
 
 	FGameplayTagSource::SetFavoriteName(bWasFavorite ? NAME_None : ActiveTagSource);
@@ -185,7 +189,8 @@ FReply SAddNewGameplayTagWidget::OnToggleTagSourceFavoriteClicked()
 
 const FSlateBrush* SAddNewGameplayTagWidget::OnGetTagSourceFavoriteImage() const
 {
-	const FName ActiveTagSource = *TagSourcesComboBox->GetSelectedItem().Get();
+	const bool bHasSelectedItem = TagSourcesComboBox.IsValid() && TagSourcesComboBox->GetSelectedItem().IsValid();
+	const FName ActiveTagSource = bHasSelectedItem ? *TagSourcesComboBox->GetSelectedItem().Get() : FName();
 	const bool bIsFavoriteTagSource = FGameplayTagSource::GetFavoriteName() == ActiveTagSource;
 
 	return FAppStyle::GetBrush(bIsFavoriteTagSource ? TEXT("Icons.Star") : TEXT("PropertyWindow.Favorites_Disabled"));
@@ -202,10 +207,10 @@ void SAddNewGameplayTagWidget::Tick( const FGeometry& AllottedGeometry, const do
 
 void SAddNewGameplayTagWidget::PopulateTagSources()
 {
-	UGameplayTagsManager& Manager = UGameplayTagsManager::Get();
+	const UGameplayTagsManager& Manager = UGameplayTagsManager::Get();
 	TagSources.Empty();
 
-	FName DefaultSource = FGameplayTagSource::GetDefaultName();
+	const FName DefaultSource = FGameplayTagSource::GetDefaultName();
 
 	// Always ensure that the default source is first
 	TagSources.Add( MakeShareable( new FName( DefaultSource ) ) );
@@ -248,7 +253,7 @@ void SAddNewGameplayTagWidget::SetTagName(const FText& InName)
 void SAddNewGameplayTagWidget::SelectTagSource(const FName& InSource)
 {
 	// Attempt to find the location in our sources, otherwise just use the first one
-	int32 SourceIndex = 0;
+	int32 SourceIndex = INDEX_NONE;
 
 	if (!InSource.IsNone())
 	{
@@ -264,7 +269,10 @@ void SAddNewGameplayTagWidget::SelectTagSource(const FName& InSource)
 		}
 	}
 
-	TagSourcesComboBox->SetSelectedItem(TagSources[SourceIndex]);
+	if (SourceIndex != INDEX_NONE && TagSourcesComboBox.IsValid())
+	{
+		TagSourcesComboBox->SetSelectedItem(TagSources[SourceIndex]);
+	}
 }
 
 void SAddNewGameplayTagWidget::OnCommitNewTagName(const FText& InText, ETextCommit::Type InCommitType)
@@ -281,19 +289,32 @@ FReply SAddNewGameplayTagWidget::OnAddNewTagButtonPressed()
 	return FReply::Handled();
 }
 
-void SAddNewGameplayTagWidget::AddSubtagFromParent(const FString& ParentTagName, const FName& ParentTagSource)
+void SAddNewGameplayTagWidget::AddSubtagFromParent(const FString& InParentTagName, const FName& InParentTagSource)
 {
-	FText SubtagBaseName = !ParentTagName.IsEmpty() ? FText::Format(FText::FromString(TEXT("{0}.")), FText::FromString(ParentTagName)) : FText();
+	const FText SubtagBaseName = !InParentTagName.IsEmpty() ? FText::Format(FText::FromString(TEXT("{0}.")), FText::FromString(InParentTagName)) : FText();
 
 	SetTagName(SubtagBaseName);
-	SelectTagSource(ParentTagSource);
+	SelectTagSource(InParentTagSource);
+
+	bShouldGetKeyboardFocus = true;
+}
+
+void SAddNewGameplayTagWidget::AddDuplicate(const FString& InParentTagName, const FName& InParentTagSource)
+{
+	SetTagName(FText::FromString(InParentTagName));
+	SelectTagSource(InParentTagSource);
 
 	bShouldGetKeyboardFocus = true;
 }
 
 void SAddNewGameplayTagWidget::CreateNewGameplayTag()
 {
-	UGameplayTagsManager& Manager = UGameplayTagsManager::Get();
+	if (NotificationItem.IsValid())
+	{
+		NotificationItem->SetVisibility(EVisibility::Collapsed);
+	}
+
+	const UGameplayTagsManager& Manager = UGameplayTagsManager::Get();
 
 	// Only support adding tags via ini file
 	if (Manager.ShouldImportTagsFromINI() == false)
@@ -303,16 +324,28 @@ void SAddNewGameplayTagWidget::CreateNewGameplayTag()
 
 	if (TagSourcesComboBox->GetSelectedItem().Get() == nullptr)
 	{
+		FNotificationInfo Info(LOCTEXT("NoTagSource", "You must specify a source file for gameplay tags."));
+		Info.ExpireDuration = 10.f;
+		Info.bUseSuccessFailIcons = true;
+		Info.Image = FAppStyle::GetBrush(TEXT("MessageLog.Error"));
+		NotificationItem = FSlateNotificationManager::Get().AddNotification(Info);
+		
 		return;
 	}
 
-	FText TagNameAsText = TagNameTextBox->GetText();
+	const FText TagNameAsText = TagNameTextBox->GetText();
 	FString TagName = TagNameAsText.ToString();
-	FString TagComment = TagCommentTextBox->GetText().ToString();
-	FName TagSource = *TagSourcesComboBox->GetSelectedItem().Get();
+	const FString TagComment = TagCommentTextBox->GetText().ToString();
+	const FName TagSource = *TagSourcesComboBox->GetSelectedItem().Get();
 
 	if (TagName.IsEmpty())
 	{
+		FNotificationInfo Info(LOCTEXT("NoTagName", "You must specify tag name."));
+		Info.ExpireDuration = 10.f;
+		Info.bUseSuccessFailIcons = true;
+		Info.Image = FAppStyle::GetBrush(TEXT("MessageLog.Error"));
+		NotificationItem = FSlateNotificationManager::Get().AddNotification(Info);
+
 		return;
 	}
 
@@ -323,7 +356,12 @@ void SAddNewGameplayTagWidget::CreateNewGameplayTag()
 		(IsValidTag.IsBound() && !IsValidTag.Execute(TagName, &ErrorMsg))
 		)
 	{
-		FMessageDialog::Open(EAppMsgType::Ok, ErrorMsg, LOCTEXT("InvalidTag", "Invalid Tag"));
+		FNotificationInfo Info(ErrorMsg);
+		Info.ExpireDuration = 10.f;
+		Info.bUseSuccessFailIcons = true;
+		Info.Image = FAppStyle::GetBrush(TEXT("MessageLog.Error"));
+		NotificationItem = FSlateNotificationManager::Get().AddNotification(Info);
+
 		return;
 	}
 
