GameplayTags: Improved gameplay tag UI

diff --git a/Engine/Plugins/Editor/GameplayTagsEditor/Source/GameplayTagsEditor/Private/SAddNewRestrictedGameplayTagWidget.cpp b/Engine/Plugins/Editor/GameplayTagsEditor/Source/GameplayTagsEditor/Private/SAddNewRestrictedGameplayTagWidget.cpp
--- a/Engine/Plugins/Editor/GameplayTagsEditor/Source/GameplayTagsEditor/Private/SAddNewRestrictedGameplayTagWidget.cpp
+++ b/Engine/Plugins/Editor/GameplayTagsEditor/Source/GameplayTagsEditor/Private/SAddNewRestrictedGameplayTagWidget.cpp
@@ -11,6 +11,8 @@
 #include "Widgets/Notifications/SNotificationList.h"
 #include "Framework/Notifications/NotificationManager.h"
 #include "Widgets/Input/SEditableTextBox.h"
+#include "Widgets/Layout/SGridPanel.h"
+
 
 #define LOCTEXT_NAMESPACE "AddNewRestrictedGameplayTagWidget"
 
@@ -31,7 +33,6 @@ void SAddNewRestrictedGameplayTagWidget::Construct(const FArguments& InArgs)
 		HintText = FText::FromString(DefaultNewName);
 	}
 
-
 	bAddingNewRestrictedTag = false;
 	bShouldGetKeyboardFocus = false;
 
@@ -42,124 +43,119 @@ void SAddNewRestrictedGameplayTagWidget::Construct(const FArguments& InArgs)
 
 	ChildSlot
 	[
-		SNew(SVerticalBox)
-
-		// Restricted Tag Name
-		+ SVerticalBox::Slot()
-		.AutoHeight()
-		.VAlign(VAlign_Top)
+		SNew(SBorder)
+		.BorderImage(FAppStyle::GetBrush("ToolPanel.GroupBorder"))
 		[
-			SNew(SHorizontalBox)
-			+ SHorizontalBox::Slot()
-			.Padding(2.0f, 4.0f)
-			.AutoWidth()
+			SNew(SGridPanel)
+			
+			// Restricted Tag Name
+			+ SGridPanel::Slot(0, 0)
+			.Padding(5)
+			.VAlign(VAlign_Center)
+			.HAlign(HAlign_Left)
 			[
-				SNew(STextBlock)
-				.Text(LOCTEXT("NewTagName", "Name:"))
+				SNew(SBox)
+				.MinDesiredWidth(150.0f)
+				[
+					SNew(STextBlock)
+					.Font(FAppStyle::GetFontStyle( TEXT("PropertyWindow.NormalFont")))
+					.Text(LOCTEXT("NewTagName", "Name:"))
+				]
 			]
-
-			+ SHorizontalBox::Slot()
-			.Padding(2.0f, 2.0f)
-			.FillWidth(1.0f)
-			.HAlign(HAlign_Right)
+			+ SGridPanel::Slot(1, 0)
+			.Padding(5)
+			.VAlign(VAlign_Center)
+			.HAlign(HAlign_Left)
 			[
-				SAssignNew(TagNameTextBox, SEditableTextBox)
-				.MinDesiredWidth(240.0f)
-				.HintText(HintText)
-				.OnTextCommitted(this, &SAddNewRestrictedGameplayTagWidget::OnCommitNewTagName)
+				SNew(SBox)
+				.MinDesiredWidth(300.0f)
+				[
+					SAssignNew(TagNameTextBox, SEditableTextBox)
+					.HintText(HintText)
+					.OnTextCommitted(this, &SAddNewRestrictedGameplayTagWidget::OnCommitNewTagName)
+				]
 			]
-		]
-
-		// Tag Comment
-		+ SVerticalBox::Slot()
-		.AutoHeight()
-		.VAlign(VAlign_Top)
-		[
-			SNew(SHorizontalBox)
-			+ SHorizontalBox::Slot()
-			.Padding(2.0f, 4.0f)
-			.AutoWidth()
+			
+			// Tag Comment
+			+ SGridPanel::Slot(0, 1)
+			.Padding(5)
+			.VAlign(VAlign_Center)
+			.HAlign(HAlign_Left)
 			[
 				SNew(STextBlock)
+				.Font(FAppStyle::GetFontStyle( TEXT("PropertyWindow.NormalFont")))
 				.Text(LOCTEXT("TagComment", "Comment:"))
 			]
-
-			+ SHorizontalBox::Slot()
-			.Padding(2.0f, 2.0f)
-			.FillWidth(1.0f)
-			.HAlign(HAlign_Right)
+			+ SGridPanel::Slot(1, 1)
+			.Padding(5)
+			.VAlign(VAlign_Center)
+			.HAlign(HAlign_Left)
 			[
-				SAssignNew(TagCommentTextBox, SEditableTextBox)
-				.MinDesiredWidth(240.0f)
-				.HintText(LOCTEXT("TagCommentHint", "Comment"))
-				.OnTextCommitted(this, &SAddNewRestrictedGameplayTagWidget::OnCommitNewTagName)
+				SNew(SBox)
+				.MinDesiredWidth(300.0f)
+				[
+					SAssignNew(TagCommentTextBox, SEditableTextBox)
+					.HintText(LOCTEXT("TagCommentHint", "Comment"))
+					.OnTextCommitted(this, &SAddNewRestrictedGameplayTagWidget::OnCommitNewTagName)
+				]
 			]
-		]
 
-		+ SVerticalBox::Slot()
-		.AutoHeight()
-		.VAlign(VAlign_Top)
-		[
-			SNew(SHorizontalBox)
-			+ SHorizontalBox::Slot()
-			.Padding(2.0f, 4.0f)
-			.AutoWidth()
+			// Allow non-restricted children
+			+ SGridPanel::Slot(0, 2)
+			.Padding(5)
+			.VAlign(VAlign_Center)
+			.HAlign(HAlign_Left)
 			[
 				SNew(STextBlock)
+				.Font(FAppStyle::GetFontStyle( TEXT("PropertyWindow.NormalFont")))
 				.Text(LOCTEXT("AllowNonRestrictedChildren", "Allow non-restricted children:"))
 			]
-
-			+ SHorizontalBox::Slot()
-			.Padding(2.0f, 2.0f)
-			.FillWidth(1.0f)
-			.HAlign(HAlign_Right)
+			+ SGridPanel::Slot(1, 2)
+			.Padding(5)
+			.VAlign(VAlign_Center)
+			.HAlign(HAlign_Left)
 			[
-				SAssignNew(AllowNonRestrictedChildrenCheckBox, SCheckBox)
-			]
+				SNew(SBox)
+				.MinDesiredWidth(300.0f)
+				[
+					SAssignNew(AllowNonRestrictedChildrenCheckBox, SCheckBox)
+				]
 			]
 
-		// Tag Location
-		+ SVerticalBox::Slot()
-		.AutoHeight()
-		.VAlign(VAlign_Top)
-		[
-			SNew(SHorizontalBox)
-			+ SHorizontalBox::Slot()
-			.Padding(2.0f, 6.0f)
-			.AutoWidth()
+			// Tag Location
+			+ SGridPanel::Slot(0, 3)
+			.Padding(5)
+			.VAlign(VAlign_Center)
+			.HAlign(HAlign_Left)
 			[
 				SNew(STextBlock)
 				.Text(LOCTEXT("CreateTagSource", "Source:"))
+				.Font(IDetailLayoutBuilder::GetDetailFont())
 			]
-
-			+ SHorizontalBox::Slot()
-			.Padding(2.0f, 2.0f)
-			.FillWidth(1.0f)
-			.HAlign(HAlign_Right)
+			+ SGridPanel::Slot(1, 3)
+			.Padding(5)
+			.VAlign(VAlign_Center)
+			.HAlign(HAlign_Left)
 			[
-				SAssignNew(TagSourcesComboBox, SComboBox<TSharedPtr<FName> >)
-				.OptionsSource(&RestrictedTagSources)
-				.OnGenerateWidget(this, &SAddNewRestrictedGameplayTagWidget::OnGenerateTagSourcesComboBox)
-				.ContentPadding(2.0f)
-				.Content()
+				SNew(SBox)
+				.MinDesiredWidth(300.0f)
 				[
-					SNew(STextBlock)
-					.Text(this, &SAddNewRestrictedGameplayTagWidget::CreateTagSourcesComboBoxContent)
-					.Font(IDetailLayoutBuilder::GetDetailFont())
+					SAssignNew(TagSourcesComboBox, SComboBox<TSharedPtr<FName> >)
+					.OptionsSource(&RestrictedTagSources)
+					.OnGenerateWidget(this, &SAddNewRestrictedGameplayTagWidget::OnGenerateTagSourcesComboBox)
+					.Content()
+					[
+						SNew(STextBlock)
+						.Text(this, &SAddNewRestrictedGameplayTagWidget::CreateTagSourcesComboBoxContent)
+						.Font(IDetailLayoutBuilder::GetDetailFont())
+					]
 				]
 			]
-		]
 
-		// Add Tag Button
-		+SVerticalBox::Slot()
-		.AutoHeight()
-		.VAlign(VAlign_Top)
-		.HAlign(HAlign_Center)
-		.Padding(8.0f)
-		[
-			SNew(SHorizontalBox)
-			+ SHorizontalBox::Slot()
-			.AutoWidth()
+			// Add Tag Button
+			+ SGridPanel::Slot(1, 4)
+			.Padding(5)
+			.HAlign(HAlign_Left)
 			[
 				SNew(SButton)
 				.Text(LOCTEXT("AddNew", "Add New Tag"))
@@ -167,7 +163,7 @@ void SAddNewRestrictedGameplayTagWidget::Construct(const FArguments& InArgs)
 			]
 		]
 	];
-
+			
 	Reset();
 }
 
@@ -189,7 +185,7 @@ void SAddNewRestrictedGameplayTagWidget::PopulateTagSources()
 	Manager.GetRestrictedTagSources(Sources);
 
 	// Used to make sure we have a non-empty list of restricted tag sources. Not an actual source.
-	FName PlaceholderSource = NAME_None;
+	const FName PlaceholderSource = NAME_None;
 
 	// Add the placeholder source if no other sources exist
 	if (Sources.Num() == 0)
@@ -271,14 +267,22 @@ void SAddNewRestrictedGameplayTagWidget::AddSubtagFromParent(const FString& Pare
 	bShouldGetKeyboardFocus = true;
 }
 
+void SAddNewRestrictedGameplayTagWidget::AddDuplicate(const FString& ParentTagName, const FName& ParentTagSource, bool bAllowNonRestrictedChildren)
+{
+	SetTagName(FText::FromString(ParentTagName));
+	SelectTagSource(ParentTagSource);
+	SetAllowNonRestrictedChildren(bAllowNonRestrictedChildren);
+
+	bShouldGetKeyboardFocus = true;
+}
+
 void SAddNewRestrictedGameplayTagWidget::ValidateNewRestrictedTag()
 {
 	UGameplayTagsManager& Manager = UGameplayTagsManager::Get();
 
 	FString TagName = TagNameTextBox->GetText().ToString();
 	FString TagComment = TagCommentTextBox->GetText().ToString();
-	bool bAllowNonRestrictedChildren = AllowNonRestrictedChildrenCheckBox->IsChecked();
-	FName TagSource = *TagSourcesComboBox->GetSelectedItem().Get();
+	const FName TagSource = *TagSourcesComboBox->GetSelectedItem().Get();
 
 	if (TagSource == NAME_None)
 	{
@@ -353,7 +357,7 @@ void SAddNewRestrictedGameplayTagWidget::CreateNewRestrictedGameplayTag()
 		AddRestrictedGameplayTagDialog->SetVisibility(EVisibility::Collapsed);
 	}
 
-	UGameplayTagsManager& Manager = UGameplayTagsManager::Get();
+	const UGameplayTagsManager& Manager = UGameplayTagsManager::Get();
 
 	// Only support adding tags via ini file
 	if (Manager.ShouldImportTagsFromINI() == false)
@@ -361,10 +365,10 @@ void SAddNewRestrictedGameplayTagWidget::CreateNewRestrictedGameplayTag()
 		return;
 	}
 
-	FString TagName = TagNameTextBox->GetText().ToString();
-	FString TagComment = TagCommentTextBox->GetText().ToString();
-	bool bAllowNonRestrictedChildren = AllowNonRestrictedChildrenCheckBox->IsChecked();
-	FName TagSource = *TagSourcesComboBox->GetSelectedItem().Get();
+	const FString TagName = TagNameTextBox->GetText().ToString();
+	const FString TagComment = TagCommentTextBox->GetText().ToString();
+	const bool bAllowNonRestrictedChildren = AllowNonRestrictedChildrenCheckBox->IsChecked();
+	const FName TagSource = *TagSourcesComboBox->GetSelectedItem().Get();
 
 	if (TagName.IsEmpty())
 	{
@@ -383,7 +387,10 @@ void SAddNewRestrictedGameplayTagWidget::CreateNewRestrictedGameplayTag()
 
 void SAddNewRestrictedGameplayTagWidget::CancelNewTag()
 {
-	AddRestrictedGameplayTagDialog->SetVisibility(EVisibility::Collapsed);
+	if (AddRestrictedGameplayTagDialog.IsValid())
+	{
+		AddRestrictedGameplayTagDialog->SetVisibility(EVisibility::Collapsed);
+	}
 }
 
 TSharedRef<SWidget> SAddNewRestrictedGameplayTagWidget::OnGenerateTagSourcesComboBox(TSharedPtr<FName> InItem)
diff --git a/Engine/Plugins/Editor/GameplayTagsEditor/Source/GameplayTagsEditor/Private/SAddNewRestrictedGameplayTagWidget.h b/Engine/Plugins/Editor/GameplayTagsEditor/Source/GameplayTagsEditor/Private/SAddNewRestrictedGameplayTagWidget.h
--- a/Engine/Plugins/Editor/GameplayTagsEditor/Source/GameplayTagsEditor/Private/SAddNewRestrictedGameplayTagWidget.h
+++ b/Engine/Plugins/Editor/GameplayTagsEditor/Source/GameplayTagsEditor/Private/SAddNewRestrictedGameplayTagWidget.h
@@ -24,7 +24,7 @@ public:
 		SLATE_ARGUMENT( FString, NewRestrictedTagName ) // String that will initially populate the New Tag Name field
 	SLATE_END_ARGS();
 
-	virtual ~SAddNewRestrictedGameplayTagWidget();
+	virtual ~SAddNewRestrictedGameplayTagWidget() override;
 
 	virtual void Tick( const FGeometry& AllottedGeometry, const double InCurrentTime, const float InDeltaTime ) override;
 
@@ -39,6 +39,9 @@ public:
 	/** Begins the process of adding a subtag to a parent tag */
 	void AddSubtagFromParent(const FString& ParentTagName, const FName& ParentTagSource, bool bAllowNonRestrictedChildren);
 
+	/** Begins the process of adding a duplicate of existing tag */
+	void AddDuplicate(const FString& ParentTagName, const FName& ParentTagSource, bool bAllowNonRestrictedChildren);
+
 	/** Resets all input fields */
 	void Reset(FName TagSource = NAME_None);
 
@@ -82,7 +85,7 @@ private:
 	/** The name of the next gameplay tag to create */
 	TSharedPtr<SEditableTextBox> TagNameTextBox;
 
-	/** The comment to asign to the next gameplay tag to create*/
+	/** The comment to assign to the next gameplay tag to create*/
 	TSharedPtr<SEditableTextBox> TagCommentTextBox;
 
 	TSharedPtr<SCheckBox> AllowNonRestrictedChildrenCheckBox;
@@ -93,10 +96,10 @@ private:
 	/** Callback for when a new restricted gameplay tag has been added to the INI files */
 	FOnRestrictedGameplayTagAdded OnRestrictedGameplayTagAdded;
 
-	bool bAddingNewRestrictedTag;
+	bool bAddingNewRestrictedTag = false;
 
 	/** Tracks if this widget should get keyboard focus */
-	bool bShouldGetKeyboardFocus;
+	bool bShouldGetKeyboardFocus = false;
 
 	FString DefaultNewName;
 
diff --git a/Engine/Plugins/Editor/GameplayTagsEditor/Source/GameplayTagsEditor/Private/SGameplayTagContainerGraphPin.cpp b/Engine/Plugins/Editor/GameplayTagsEditor/Source/GameplayTagsEditor/Private/SGameplayTagContainerGraphPin.cpp
--- a/Engine/Plugins/Editor/GameplayTagsEditor/Source/GameplayTagsEditor/Private/SGameplayTagContainerGraphPin.cpp
+++ b/Engine/Plugins/Editor/GameplayTagsEditor/Source/GameplayTagsEditor/Private/SGameplayTagContainerGraphPin.cpp
@@ -2,59 +2,63 @@
 
 #include "SGameplayTagContainerGraphPin.h"
 #include "EdGraph/EdGraphSchema.h"
-#include "GameplayTagPinUtilities.h"
 #include "GameplayTagContainer.h"
-#include "Widgets/SBoxPanel.h"
+#include "SGameplayTagContainerCombo.h"
+#include "GameplayTagEditorUtilities.h"
+#include "ScopedTransaction.h"
 
-#define LOCTEXT_NAMESPACE "GameplayTagGraphPin"
+#define LOCTEXT_NAMESPACE "GameplayTagContainerGraphPin"
 
 void SGameplayTagContainerGraphPin::Construct(const FArguments& InArgs, UEdGraphPin* InGraphPinObj)
 {
-	SGameplayTagGraphPin::Construct(SGameplayTagGraphPin::FArguments(), InGraphPinObj);
+	SGraphPin::Construct(SGraphPin::FArguments(), InGraphPinObj);
 }
 
 void SGameplayTagContainerGraphPin::ParseDefaultValueData()
 {
-	FilterString = GameplayTagPinUtilities::ExtractTagFilterStringFromGraphPin(GraphPinObj);
-
 	// Read using import text, but with serialize flag set so it doesn't always throw away invalid ones
-	TagContainer->FromExportString(GraphPinObj->GetDefaultAsString(), PPF_SerializedAsImportText);
+	GameplayTagContainer.FromExportString(GraphPinObj->GetDefaultAsString(), PPF_SerializedAsImportText);
 }
 
-TSharedRef<SWidget> SGameplayTagContainerGraphPin::GetEditContent()
+void SGameplayTagContainerGraphPin::OnTagContainerChanged(const FGameplayTagContainer& NewTagContainer)
 {
-	EditableContainers.Empty();
-	EditableContainers.Add( SGameplayTagWidget::FEditableGameplayTagContainerDatum( GraphPinObj->GetOwningNode(), TagContainer.Get() ) );
-
-	return SNew( SVerticalBox )
-		+SVerticalBox::Slot()
-		.AutoHeight()
-		.MaxHeight( 400 )
-		[
-			SNew( SGameplayTagWidget, EditableContainers )
-			.OnTagChanged(this, &SGameplayTagContainerGraphPin::SaveDefaultValueData)
-			.TagContainerName( TEXT("SGameplayTagContainerGraphPin") )
-			.Visibility( this, &SGraphPin::GetDefaultValueVisibility )
-			.MultiSelect(true)
-			.Filter(FilterString)
-		];
-}
+	GameplayTagContainer = NewTagContainer;
 
-void SGameplayTagContainerGraphPin::SaveDefaultValueData()
-{	
-	RefreshCachedData();
+	const FString TagContainerString = UE::GameplayTags::EditorUtilities::GameplayTagContainerExportText(GameplayTagContainer);
 
-	// Set Pin Data
-	FString TagContainerString = TagContainer->ToString();
 	FString CurrentDefaultValue = GraphPinObj->GetDefaultAsString();
 	if (CurrentDefaultValue.IsEmpty())
 	{
-		CurrentDefaultValue = FString(TEXT("(GameplayTags=)"));
+		CurrentDefaultValue = UE::GameplayTags::EditorUtilities::GameplayTagContainerExportText(FGameplayTagContainer());
 	}
+			
 	if (!CurrentDefaultValue.Equals(TagContainerString))
 	{
+		const FScopedTransaction Transaction(LOCTEXT("ChangeDefaultValue", "Change Pin Default Value"));
 		GraphPinObj->GetSchema()->TrySetDefaultValue(*GraphPinObj, TagContainerString);
 	}
 }
 
+TSharedRef<SWidget> SGameplayTagContainerGraphPin::GetDefaultValueWidget()
+{
+	if (GraphPinObj == nullptr)
+	{
+		return SNullWidget::NullWidget;
+	}
+
+	ParseDefaultValueData();
+	const FString FilterString = UE::GameplayTags::EditorUtilities::ExtractTagFilterStringFromGraphPin(GraphPinObj);
+
+	return SNew(SGameplayTagContainerCombo)
+		.Visibility(this, &SGraphPin::GetDefaultValueVisibility)
+		.Filter(FilterString)
+		.TagContainer(this, &SGameplayTagContainerGraphPin::GetTagContainer)
+		.OnTagContainerChanged(this, &SGameplayTagContainerGraphPin::OnTagContainerChanged);
+}
+
+FGameplayTagContainer SGameplayTagContainerGraphPin::GetTagContainer() const
+{
+	return GameplayTagContainer;
+}
+
 #undef LOCTEXT_NAMESPACE
diff --git a/Engine/Plugins/Editor/GameplayTagsEditor/Source/GameplayTagsEditor/Private/SGameplayTagContainerGraphPin.h b/Engine/Plugins/Editor/GameplayTagsEditor/Source/GameplayTagsEditor/Private/SGameplayTagContainerGraphPin.h
--- a/Engine/Plugins/Editor/GameplayTagsEditor/Source/GameplayTagsEditor/Private/SGameplayTagContainerGraphPin.h
+++ b/Engine/Plugins/Editor/GameplayTagsEditor/Source/GameplayTagsEditor/Private/SGameplayTagContainerGraphPin.h
@@ -2,10 +2,11 @@
 
 #pragma once
 
-#include "SGameplayTagGraphPin.h"
+#include "KismetPins/SGraphPinStructInstance.h"
+#include "GameplayTagContainer.h"
 
 /** Almost the same as a tag pin, but supports multiple tags */
-class SGameplayTagContainerGraphPin : public SGameplayTagGraphPin
+class SGameplayTagContainerGraphPin : public SGraphPinStructInstance
 {
 public:
 	SLATE_BEGIN_ARGS(SGameplayTagContainerGraphPin) {}
@@ -16,7 +17,11 @@ public:
 protected:
 	//~ Begin SGameplayTagGraphPin Interface
 	virtual void ParseDefaultValueData() override;
-	virtual void SaveDefaultValueData() override;
-	virtual TSharedRef<SWidget> GetEditContent() override;
+	virtual TSharedRef<SWidget> GetDefaultValueWidget() override;
 	//~ End SGraphPin Interface
+
+	FGameplayTagContainer GetTagContainer() const;
+	void OnTagContainerChanged(const FGameplayTagContainer& NewTagContainer);
+	
+	FGameplayTagContainer GameplayTagContainer;
 };
diff --git a/Engine/Plugins/Editor/GameplayTagsEditor/Source/GameplayTagsEditor/Private/SGameplayTagGraphPin.cpp b/Engine/Plugins/Editor/GameplayTagsEditor/Source/GameplayTagsEditor/Private/SGameplayTagGraphPin.cpp
--- a/Engine/Plugins/Editor/GameplayTagsEditor/Source/GameplayTagsEditor/Private/SGameplayTagGraphPin.cpp
+++ b/Engine/Plugins/Editor/GameplayTagsEditor/Source/GameplayTagsEditor/Private/SGameplayTagGraphPin.cpp
@@ -2,120 +2,63 @@
 
 #include "SGameplayTagGraphPin.h"
 #include "EdGraph/EdGraphSchema.h"
-#include "GameplayTagPinUtilities.h"
-#include "Framework/Views/TableViewMetadata.h"
 #include "GameplayTagContainer.h"
-#include "Widgets/Views/SListView.h"
+#include "GameplayTagEditorUtilities.h"
+#include "SGameplayTagCombo.h"
+#include "ScopedTransaction.h"
 
 #define LOCTEXT_NAMESPACE "GameplayTagGraphPin"
 
 void SGameplayTagGraphPin::Construct( const FArguments& InArgs, UEdGraphPin* InGraphPinObj )
 {
-	TagContainer = MakeShareable( new FGameplayTagContainer() );
 	SGraphPin::Construct( SGraphPin::FArguments(), InGraphPinObj );
 }
 
 void SGameplayTagGraphPin::ParseDefaultValueData()
 {
-	FilterString = GameplayTagPinUtilities::ExtractTagFilterStringFromGraphPin(GraphPinObj);
-
-	FGameplayTag GameplayTag;
-	
 	// Read using import text, but with serialize flag set so it doesn't always throw away invalid ones
 	GameplayTag.FromExportString(GraphPinObj->GetDefaultAsString(), PPF_SerializedAsImportText);
-
-	if (GameplayTag.IsValid())
-	{
-		TagContainer->AddTag(GameplayTag);
-	}
 }
 
-TSharedRef<SWidget> SGameplayTagGraphPin::GetEditContent()
+void SGameplayTagGraphPin::OnTagChanged(const FGameplayTag NewTag)
 {
-	EditableContainers.Empty();
-	EditableContainers.Add( SGameplayTagWidget::FEditableGameplayTagContainerDatum( GraphPinObj->GetOwningNode(), TagContainer.Get() ) );
+	GameplayTag = NewTag;
 
-	return SNew( SVerticalBox )
-		+SVerticalBox::Slot()
-		.AutoHeight()
-		.MaxHeight( 400 )
-		[
-			SNew( SGameplayTagWidget, EditableContainers )
-			.OnTagChanged( this, &SGameplayTagGraphPin::SaveDefaultValueData)
-			.TagContainerName( TEXT("SGameplayTagGraphPin") )
-			.Visibility( this, &SGraphPin::GetDefaultValueVisibility )
-			.MultiSelect(false)
-			.Filter(FilterString)
-		];
-}
+	const FString TagString = UE::GameplayTags::EditorUtilities::GameplayTagExportText(GameplayTag);
 
-TSharedRef<SWidget> SGameplayTagGraphPin::GetDescriptionContent()
-{
-	RefreshCachedData();
-
-	SAssignNew( TagListView, SListView<TSharedPtr<FString>> )
-		.ListItemsSource(&TagNames)
-		.SelectionMode(ESelectionMode::None)
-		.OnGenerateRow(this, &SGameplayTagGraphPin::OnGenerateRow);
-
-	return TagListView->AsShared();
-}
-
-TSharedRef<ITableRow> SGameplayTagGraphPin::OnGenerateRow(TSharedPtr<FString> Item, const TSharedRef<STableViewBase>& OwnerTable)
-{
-	return SNew( STableRow< TSharedPtr<FString> >, OwnerTable )
-		[
-			SNew(STextBlock) .Text( FText::FromString(*Item.Get()) )
-		];
-}
-
-void SGameplayTagGraphPin::RefreshCachedData()
-{
-	// Clear the list
-	TagNames.Empty();
-
-	// Add tags to list
-	FString TagName;
-	if (TagContainer.IsValid())
+	FString CurrentDefaultValue = GraphPinObj->GetDefaultAsString();
+	if (CurrentDefaultValue.IsEmpty())
 	{
-		for (auto It = TagContainer->CreateConstIterator(); It; ++It)
-		{
-			TagName = It->ToString();
-			TagNames.Add(MakeShareable(new FString(TagName)));
-		}
+		CurrentDefaultValue = UE::GameplayTags::EditorUtilities::GameplayTagExportText(FGameplayTag());
 	}
-
-	// Refresh the slate list
-	if (TagListView.IsValid())
+			
+	if (!CurrentDefaultValue.Equals(TagString))
 	{
-		TagListView->RequestListRefresh();
+		const FScopedTransaction Transaction(LOCTEXT("ChangeDefaultValue", "Change Pin Default Value"));
+		GraphPinObj->GetSchema()->TrySetDefaultValue(*GraphPinObj, TagString);
 	}
 }
 
-void SGameplayTagGraphPin::SaveDefaultValueData()
-{	
-	RefreshCachedData();
-
-	// Set Pin Data
-	FString TagString;
-
-	if (TagNames.Num() > 0 && TagNames[0].IsValid())
-	{
-		TagString = TEXT("(");
-		TagString += TEXT("TagName=\"");
-		TagString += *TagNames[0].Get();
-		TagString += TEXT("\"");
-		TagString += TEXT(")");
-	}
-	FString CurrentDefaultValue = GraphPinObj->GetDefaultAsString();
-	if (CurrentDefaultValue.IsEmpty() || CurrentDefaultValue == TEXT("(TagName=\"\")"))
-	{
-		CurrentDefaultValue = FString(TEXT(""));
-	}
-	if (!CurrentDefaultValue.Equals(TagString))
+TSharedRef<SWidget> SGameplayTagGraphPin::GetDefaultValueWidget()
+{
+	if (GraphPinObj == nullptr)
 	{
-		GraphPinObj->GetSchema()->TrySetDefaultValue(*GraphPinObj, TagString);
+		return SNullWidget::NullWidget;
 	}
+
+	ParseDefaultValueData();
+	const FString FilterString = UE::GameplayTags::EditorUtilities::ExtractTagFilterStringFromGraphPin(GraphPinObj);
+
+	return SNew(SGameplayTagCombo)
+		.Visibility(this, &SGraphPin::GetDefaultValueVisibility)
+		.Filter(FilterString)
+		.Tag(this, &SGameplayTagGraphPin::GetGameplayTag)
+		.OnTagChanged(this, &SGameplayTagGraphPin::OnTagChanged);
+}
+
+FGameplayTag SGameplayTagGraphPin::GetGameplayTag() const
+{
+	return GameplayTag;
 }
 
 #undef LOCTEXT_NAMESPACE
diff --git a/Engine/Plugins/Editor/GameplayTagsEditor/Source/GameplayTagsEditor/Private/SGameplayTagGraphPin.h b/Engine/Plugins/Editor/GameplayTagsEditor/Source/GameplayTagsEditor/Private/SGameplayTagGraphPin.h
--- a/Engine/Plugins/Editor/GameplayTagsEditor/Source/GameplayTagsEditor/Private/SGameplayTagGraphPin.h
+++ b/Engine/Plugins/Editor/GameplayTagsEditor/Source/GameplayTagsEditor/Private/SGameplayTagGraphPin.h
@@ -3,7 +3,7 @@
 #pragma once
 
 #include "KismetPins/SGraphPinStructInstance.h"
-#include "SGameplayTagWidget.h"
+#include "GameplayTagContainer.h"
 
 template <typename ItemType> class SListView;
 
@@ -19,30 +19,11 @@ public:
 protected:
 	//~ Begin SGraphPinStructInstance Interface
 	virtual void ParseDefaultValueData() override;
-	virtual void SaveDefaultValueData() override;
-	virtual void RefreshCachedData() override;
-	virtual TSharedRef<SWidget> GetEditContent() override;
-	virtual TSharedRef<SWidget> GetDescriptionContent() override;
+	virtual TSharedRef<SWidget> GetDefaultValueWidget() override;
 	//~ End SGraphPin Interface
 
-	/** 
-	 * Callback for populating rows of the SelectedTags List View.
-	 * @return widget that contains the name of a tag.
-	 */
-	TSharedRef<ITableRow> OnGenerateRow(TSharedPtr<FString> Item, const TSharedRef<STableViewBase>& OwnerTable);
+	FGameplayTag GetGameplayTag() const;
+	void OnTagChanged(const FGameplayTag NewTag);
 
-	// Tag Container used for the GameplayTagWidget.
-	TSharedPtr<FGameplayTagContainer> TagContainer;
-
-	// Datum uses for the GameplayTagWidget.
-	TArray<SGameplayTagWidget::FEditableGameplayTagContainerDatum> EditableContainers;
-
-	// Array of names for the read only display of tag names on the node.
-	TArray< TSharedPtr<FString> > TagNames;
-
-	// The List View used to display the read only tag names on the node.
-	TSharedPtr<SListView<TSharedPtr<FString>>> TagListView;
-
-	// Filter the list of available tags
-	FString FilterString;
+	FGameplayTag GameplayTag;
 };
diff --git a/Engine/Plugins/Editor/GameplayTagsEditor/Source/GameplayTagsEditor/Private/SGameplayTagQueryGraphPin.cpp b/Engine/Plugins/Editor/GameplayTagsEditor/Source/GameplayTagsEditor/Private/SGameplayTagQueryGraphPin.cpp
--- a/Engine/Plugins/Editor/GameplayTagsEditor/Source/GameplayTagsEditor/Private/SGameplayTagQueryGraphPin.cpp
+++ b/Engine/Plugins/Editor/GameplayTagsEditor/Source/GameplayTagsEditor/Private/SGameplayTagQueryGraphPin.cpp
@@ -2,106 +2,53 @@
 
 #include "SGameplayTagQueryGraphPin.h"
 #include "EdGraph/EdGraphSchema.h"
-#include "Widgets/Input/SComboButton.h"
 #include "GameplayTagContainer.h"
-#include "Widgets/Layout/SScaleBox.h"
 #include "ScopedTransaction.h"
+#include "GameplayTagEditorUtilities.h"
+#include "SGameplayTagQueryEntryBox.h"
 
 #define LOCTEXT_NAMESPACE "GameplayTagQueryGraphPin"
 
+
 void SGameplayTagQueryGraphPin::Construct(const FArguments& InArgs, UEdGraphPin* InGraphPinObj)
 {
-	TagQuery = MakeShareable(new FGameplayTagQuery());
 	SGraphPin::Construct( SGraphPin::FArguments(), InGraphPinObj );
 }
 
-TSharedRef<SWidget>	SGameplayTagQueryGraphPin::GetDefaultValueWidget()
+void SGameplayTagQueryGraphPin::OnTagQueryChanged(const FGameplayTagQuery& NewTagQuery)
 {
-	ParseDefaultValueData();
-
-	//Create widget
-	return SNew(SVerticalBox)
-		+SVerticalBox::Slot()
-		.AutoHeight()
-		[
-			SAssignNew( ComboButton, SComboButton )
-			.OnGetMenuContent(this, &SGameplayTagQueryGraphPin::GetListContent)
-			.ContentPadding( FMargin( 2.0f, 2.0f ) )
-			.Visibility( this, &SGraphPin::GetDefaultValueVisibility )
-			.ButtonContent()
-			[
-				SNew( STextBlock )
-				.Text( LOCTEXT("GameplayTagQueryWidget_Edit", "Edit") )
-			]
-		]
-		+SVerticalBox::Slot()
-		.AutoHeight()
-		[
-			QueryDesc()
-		];
-}
-
-void SGameplayTagQueryGraphPin::ParseDefaultValueData()
-{
-	FString const TagQueryString = GraphPinObj->GetDefaultAsString();
-
-	FProperty* const TQProperty = FindFProperty<FProperty>(UEditableGameplayTagQuery::StaticClass(), TEXT("TagQueryExportText_Helper"));
-	if (TQProperty)
-	{
-		FGameplayTagQuery* const TQ = TagQuery.Get();
-		TQProperty->ImportText_Direct(*TagQueryString, TQ, nullptr, 0, GLog);
-	}
-}
+	TagQuery = NewTagQuery;
+			
+	const FString TagQueryExportText = UE::GameplayTags::EditorUtilities::GameplayTagQueryExportText(TagQuery);
 
-TSharedRef<SWidget> SGameplayTagQueryGraphPin::GetListContent()
-{
-	EditableQueries.Empty();
-	EditableQueries.Add(SGameplayTagQueryWidget::FEditableGameplayTagQueryDatum(GraphPinObj->GetOwningNode(), TagQuery.Get(), &TagQueryExportText));
-
-	return SNew( SVerticalBox )
-		+SVerticalBox::Slot()
-		.AutoHeight()
-		.MaxHeight( 400 )
-		[
-			SNew(SScaleBox)
-			.HAlign(EHorizontalAlignment::HAlign_Left)
-			.VAlign(EVerticalAlignment::VAlign_Top)
-			.StretchDirection(EStretchDirection::DownOnly)
-			.Stretch(EStretch::ScaleToFit)
-			.Content()
-			[
-				SNew(SGameplayTagQueryWidget, EditableQueries, nullptr)
-				.OnQueryChanged(this, &SGameplayTagQueryGraphPin::OnQueryChanged)
-				.Visibility( this, &SGraphPin::GetDefaultValueVisibility )
-				.AutoSave(true)
-			]
-		];
-}
-
-void SGameplayTagQueryGraphPin::OnQueryChanged()
-{
 	if (TagQueryExportText != GraphPinObj->GetDefaultAsString())
 	{
 		// Set Pin Data
 		const FScopedTransaction Transaction(NSLOCTEXT("GraphEditor", "ChangePinValue", "Change Pin Value"));
-		GraphPinObj->Modify();
 		GraphPinObj->GetSchema()->TrySetDefaultValue(*GraphPinObj, TagQueryExportText);
 	}
-	QueryDescription = TagQuery->GetDescription();
 }
 
-TSharedRef<SWidget> SGameplayTagQueryGraphPin::QueryDesc()
+TSharedRef<SWidget>	SGameplayTagQueryGraphPin::GetDefaultValueWidget()
 {
-	QueryDescription = TagQuery->GetDescription();
+	ParseDefaultValueData();
+
+	return SNew(SGameplayTagQueryEntryBox)
+		.Visibility(this, &SGraphPin::GetDefaultValueVisibility)
+		.DescriptionMaxWidth(250.0f)
+		.TagQuery(this, &SGameplayTagQueryGraphPin::GetTagQuery)
+		.OnTagQueryChanged(this, &SGameplayTagQueryGraphPin::OnTagQueryChanged);
+}
 
-	return SNew(STextBlock)
-		.Text(this, &SGameplayTagQueryGraphPin::GetQueryDescText)
-		.AutoWrapText(true);
+FGameplayTagQuery SGameplayTagQueryGraphPin::GetTagQuery() const
+{
+	return TagQuery;
 }
 
-FText SGameplayTagQueryGraphPin::GetQueryDescText() const
+void SGameplayTagQueryGraphPin::ParseDefaultValueData()
 {
-	return FText::FromString(QueryDescription);
+	FString const TagQueryString = GraphPinObj->GetDefaultAsString();
+	TagQuery = UE::GameplayTags::EditorUtilities::GameplayTagQueryTryImportText(TagQueryString);
 }
 
 #undef LOCTEXT_NAMESPACE
diff --git a/Engine/Plugins/Editor/GameplayTagsEditor/Source/GameplayTagsEditor/Private/SGameplayTagQueryGraphPin.h b/Engine/Plugins/Editor/GameplayTagsEditor/Source/GameplayTagsEditor/Private/SGameplayTagQueryGraphPin.h
--- a/Engine/Plugins/Editor/GameplayTagsEditor/Source/GameplayTagsEditor/Private/SGameplayTagQueryGraphPin.h
+++ b/Engine/Plugins/Editor/GameplayTagsEditor/Source/GameplayTagsEditor/Private/SGameplayTagQueryGraphPin.h
@@ -3,7 +3,7 @@
 #pragma once
 
 #include "SGraphPin.h"
-#include "SGameplayTagQueryWidget.h"
+#include "GameplayTagContainer.h"
 
 class SComboButton;
 
@@ -20,35 +20,12 @@ public:
 	//~ End SGraphPin Interface
 
 private:
-
 	/** Parses the Data from the pin to fill in the names of the array. */
 	void ParseDefaultValueData();
 
-	/** Callback function to create content for the combo button. */
-	TSharedRef<SWidget> GetListContent();
-
-	void OnQueryChanged();
-
-	/**
-	 * Creates SelectedTags List.
-	 * @return widget that contains the read only tag names for displaying on the node.
-	 */
-	TSharedRef<SWidget> QueryDesc();
-
-private:
-
-	// Combo Button for the drop down list.
-	TSharedPtr<SComboButton> ComboButton;
-
-	// Tag Container used for the GameplayTagWidget.
-	TSharedPtr<FGameplayTagQuery> TagQuery;
-
-	FString TagQueryExportText;
-
-	// Datum uses for the GameplayTagWidget.
-	TArray<SGameplayTagQueryWidget::FEditableGameplayTagQueryDatum> EditableQueries;
-
-	FText GetQueryDescText() const;
+	FGameplayTagQuery GetTagQuery() const;
+	void OnTagQueryChanged(const FGameplayTagQuery& NewTagQuery);
 
-	FString QueryDescription;
+	/** Parse tag query used for editing. */
+	FGameplayTagQuery TagQuery;
 };
diff --git a/Engine/Plugins/Editor/GameplayTagsEditor/Source/GameplayTagsEditor/Private/SGameplayTagQueryWidget.cpp b/Engine/Plugins/Editor/GameplayTagsEditor/Source/GameplayTagsEditor/Private/SGameplayTagQueryWidget.cpp
--- a/Engine/Plugins/Editor/GameplayTagsEditor/Source/GameplayTagsEditor/Private/SGameplayTagQueryWidget.cpp
+++ b/Engine/Plugins/Editor/GameplayTagsEditor/Source/GameplayTagsEditor/Private/SGameplayTagQueryWidget.cpp
@@ -8,37 +8,26 @@
 #include "Widgets/Input/SButton.h"
 #include "IDetailsView.h"
 #include "GameplayTagsManager.h"
+#include "Widgets/Layout/SUniformGridPanel.h"
+#include "ScopedTransaction.h"
+#include "Framework/Application/SlateApplication.h"
+#include "Framework/Docking/TabManager.h"
 
 #define LOCTEXT_NAMESPACE "GameplayTagQueryWidget"
 
-void SGameplayTagQueryWidget::Construct(const FArguments& InArgs, const TArray<FEditableGameplayTagQueryDatum>& EditableTagQueries, const TSharedPtr<IPropertyHandle> InTagQueryPropertyHandle)
+void SGameplayTagQueryWidget::Construct(const FArguments& InArgs, const TArray<FGameplayTagQuery>& InTagQueries)
 {
-	ensure(EditableTagQueries.Num() > 0);
-	TagQueries = EditableTagQueries;
-	TagQueryPropertyHandle = InTagQueryPropertyHandle;
+	ensure(InTagQueries.Num() > 0);
+	TagQueries = InTagQueries;
 
 	bReadOnly = InArgs._ReadOnly;
-	bAutoSave = InArgs._AutoSave;
-	OnClosePreSave = InArgs._OnClosePreSave;
-	OnSaveAndClose = InArgs._OnSaveAndClose;
-	OnCancel = InArgs._OnCancel;
-	OnQueryChanged = InArgs._OnQueryChanged;
-
-	// Tag the assets as transactional so they can support undo/redo
-	for (int32 AssetIdx = 0; AssetIdx < TagQueries.Num(); ++AssetIdx)
-	{
-		UObject* TagQueryOwner = TagQueries[AssetIdx].TagQueryOwner.Get();
-		if (TagQueryOwner)
-		{
-			TagQueryOwner->SetFlags(RF_Transactional);
-		}
-	}
+	OnQueriesCommitted = InArgs._OnQueriesCommitted;
 
 	UGameplayTagsManager::Get().OnGetCategoriesMetaFromPropertyHandle.AddSP(this, &SGameplayTagQueryWidget::OnGetCategoriesMetaFromPropertyHandle);
 
 	// build editable query object tree from the runtime query data
-	UEditableGameplayTagQuery* const EQ = CreateEditableQuery(*TagQueries[0].TagQuery);
-	EditableQuery = EQ;
+	UEditableGameplayTagQuery* const EditableGameplayTagQuery = CreateEditableQuery(TagQueries[0]);
+	EditableQuery = EditableGameplayTagQuery;
 
 	// create details view for the editable query object
 	FDetailsViewArgs ViewArgs;
@@ -48,8 +37,7 @@ void SGameplayTagQueryWidget::Construct(const FArguments& InArgs, const TArray<F
 
 	FPropertyEditorModule& PropertyModule = FModuleManager::LoadModuleChecked<FPropertyEditorModule>("PropertyEditor");
 	Details = PropertyModule.CreateDetailView(ViewArgs);
-	Details->SetObject(EQ);
-	Details->OnFinishedChangingProperties().AddSP(this, &SGameplayTagQueryWidget::OnFinishedChangingProperties);
+	Details->SetObject(EditableGameplayTagQuery);
 
 	ChildSlot
 	[
@@ -57,39 +45,51 @@ void SGameplayTagQueryWidget::Construct(const FArguments& InArgs, const TArray<F
 		.BorderImage(FAppStyle::GetBrush("ToolPanel.GroupBorder"))
 		[
 			SNew(SVerticalBox)
+
+			// to delete!
+			+ SVerticalBox::Slot()
+			.FillHeight(1)
+			.Padding(2)
+			[
+				Details.ToSharedRef()
+			]
+
 			+ SVerticalBox::Slot()
 			.AutoHeight()
-			.VAlign(VAlign_Top)
+			.HAlign(HAlign_Right)
+			.Padding(5)
 			[
-				SNew(SHorizontalBox)
-				+ SHorizontalBox::Slot()
-				.AutoWidth()
+				SNew(SUniformGridPanel)
+				.MinDesiredSlotHeight(FAppStyle::Get().GetFloat("StandardDialog.MinDesiredSlotHeight"))
+				.MinDesiredSlotWidth(FAppStyle::Get().GetFloat("StandardDialog.MinDesiredSlotWidth"))
+				.SlotPadding(FAppStyle::Get().GetMargin("StandardDialog.SlotPadding"))
+
+				+ SUniformGridPanel::Slot(0, 0)
 				[
+					// ok button
 					SNew(SButton)
-					.IsEnabled(!bReadOnly)
-					.Visibility(this, &SGameplayTagQueryWidget::GetSaveAndCloseButtonVisibility)
-					.OnClicked(this, &SGameplayTagQueryWidget::OnSaveAndCloseClicked)
-					.Text(LOCTEXT("GameplayTagQueryWidget_SaveAndClose", "Save and Close"))
+						.IsEnabled(!bReadOnly)
+						.HAlign(HAlign_Center)
+						.OnClicked(this, &SGameplayTagQueryWidget::OnOkClicked)
+						.Text(LOCTEXT("GameplayTagQueryWidget_OK", "OK"))
 				]
-				+ SHorizontalBox::Slot()
-				.AutoWidth()
+
+				+ SUniformGridPanel::Slot(1, 0)
 				[
+					// cancel button
 					SNew(SButton)
-					.Visibility(this, &SGameplayTagQueryWidget::GetCancelButtonVisibility)
-					.OnClicked(this, &SGameplayTagQueryWidget::OnCancelClicked)
-					.Text(LOCTEXT("GameplayTagQueryWidget_Cancel", "Close Without Saving"))
+						.ContentPadding(FAppStyle::Get().GetMargin("StandardDialog.ContentPadding") )
+						.HAlign(HAlign_Center)
+						.OnClicked(this, &SGameplayTagQueryWidget::OnCancelClicked)
+						.Text(LOCTEXT("GameplayTagQueryWidget_Cancel", "Cancel"))
 				]
 			]
-			// to delete!
-			+ SVerticalBox::Slot()
-			[
-				Details.ToSharedRef()
-			]
 		]
 	];
+	
 }
 
-void SGameplayTagQueryWidget::OnGetCategoriesMetaFromPropertyHandle(TSharedPtr<IPropertyHandle> PropertyHandle, FString& MetaString)
+void SGameplayTagQueryWidget::OnGetCategoriesMetaFromPropertyHandle(TSharedPtr<IPropertyHandle> PropertyHandle, FString& MetaString) const
 {
 	TArray<UObject*> OuterObjects;
 	PropertyHandle->GetOuterObjects(OuterObjects);
@@ -101,7 +101,7 @@ void SGameplayTagQueryWidget::OnGetCategoriesMetaFromPropertyHandle(TSharedPtr<I
 			if (Outer == EditableQuery)
 			{
 				// This is us, re-route
-				MetaString = UGameplayTagsManager::StaticGetCategoriesMetaFromPropertyHandle(TagQueryPropertyHandle);
+				MetaString = Filter;
 				return;
 			}
 			Outer = Outer->GetOuter();
@@ -109,27 +109,6 @@ void SGameplayTagQueryWidget::OnGetCategoriesMetaFromPropertyHandle(TSharedPtr<I
 	}
 }
 
-void SGameplayTagQueryWidget::OnFinishedChangingProperties(const FPropertyChangedEvent& PropertyChangedEvent)
-{
-	// Auto saved changes will not call pre and post notify; auto save should only be used to make changes coming from blueprints
-	if (bAutoSave)
-	{
-		SaveToTagQuery();
-	}
-
-	OnQueryChanged.ExecuteIfBound();
-}
-
-EVisibility SGameplayTagQueryWidget::GetSaveAndCloseButtonVisibility() const
-{
-	return bAutoSave ? EVisibility::Collapsed : EVisibility::Visible;
-}
-
-EVisibility SGameplayTagQueryWidget::GetCancelButtonVisibility() const
-{
-	return bAutoSave ? EVisibility::Collapsed : EVisibility::Visible;
-}
-
 UEditableGameplayTagQuery* SGameplayTagQueryWidget::CreateEditableQuery(FGameplayTagQuery& Q)
 {
 	UEditableGameplayTagQuery* const AnEditableQuery = Q.CreateEditableQuery();
@@ -145,10 +124,10 @@ UEditableGameplayTagQuery* SGameplayTagQueryWidget::CreateEditableQuery(FGamepla
 SGameplayTagQueryWidget::~SGameplayTagQueryWidget()
 {
 	// clean up our temp editing uobjects
-	UEditableGameplayTagQuery* const Q = EditableQuery.Get();
-	if (Q)
+	UEditableGameplayTagQuery* const TagQuery = EditableQuery.Get();
+	if (TagQuery)
 	{
-		Q->RemoveFromRoot();
+		TagQuery->RemoveFromRoot();
 	}
 	
 	UGameplayTagsManager::Get().OnGetCategoriesMetaFromPropertyHandle.RemoveAll(this);
@@ -163,45 +142,38 @@ void SGameplayTagQueryWidget::SaveToTagQuery()
 	}
 
 	FGameplayTagQuery NewQuery;
-	NewQuery.BuildFromEditableQuery(*EditableQuery.Get());
-	FString NewQueryAsString = EditableQuery.Get()->GetTagQueryExportText(NewQuery);
-
-	// set query through property handle if possible to propagate changes to loaded instances
-	// Note that all queries share the same property handle so we only need to set once since it will
-	// take care of applying to all selected objects
-	IPropertyHandle* PropertyHandle = TagQueryPropertyHandle.Get();
-	if (PropertyHandle != nullptr)
-	{
-		PropertyHandle->SetValueFromFormattedString(NewQueryAsString);
-	}
 
-	// write to all selected queries
-	for (FEditableGameplayTagQueryDatum& TQ : TagQueries)
+	if (UEditableGameplayTagQuery* CurrentEditableQuery = EditableQuery.Get())
 	{
-		if (TQ.TagQueryExportText != nullptr)
+		// Do not try to build empty queries.
+		if (CurrentEditableQuery->RootExpression != nullptr)
 		{
-			*TQ.TagQueryExportText = NewQueryAsString;
+			NewQuery.BuildFromEditableQuery(*CurrentEditableQuery);
 		}
+	}
 
-		// when property handle is not available (e.g. BP pin) we set the query and dirty the package manually
-		if (PropertyHandle == nullptr)
-		{
-			*TQ.TagQuery = NewQuery;
-			if (TQ.TagQueryOwner.IsValid())
-			{
-				TQ.TagQueryOwner->MarkPackageDirty();
-			}
-		}
+	for (FGameplayTagQuery& TagQuery : TagQueries)
+	{
+		TagQuery = NewQuery;
 	}
 }
 
-FReply SGameplayTagQueryWidget::OnSaveAndCloseClicked()
+FReply SGameplayTagQueryWidget::OnOkClicked()
 {
-	OnClosePreSave.ExecuteIfBound();
-
-	SaveToTagQuery();
+	if (!bReadOnly)
+	{
+		FScopedTransaction Transaction( LOCTEXT("GameplayTagQueryWidget_Edit", "Edit Gameplay Tag Query") );
+		
+		SaveToTagQuery();
 
-	OnSaveAndClose.ExecuteIfBound();
+		OnQueriesCommitted.ExecuteIfBound(TagQueries);
+		OnOk.ExecuteIfBound();
+	}
+	else
+	{
+		OnCancel.ExecuteIfBound();
+	}
+	
 	return FReply::Handled();
 }
 
@@ -211,4 +183,94 @@ FReply SGameplayTagQueryWidget::OnCancelClicked() const
 	return FReply::Handled();
 }
 
+namespace UE::GameplayTags::Editor
+{
+
+static TWeakPtr<SGameplayTagQueryWidget> GlobalTagQueryWidget;
+static TWeakPtr<SWindow> GlobalTagQueryWidgetWindow;
+
+TWeakPtr<SGameplayTagQueryWidget> OpenGameplayTagQueryWindow(const FGameplayTagQueryWindowArgs& Args)
+{
+	CloseGameplayTagQueryWindow(nullptr);
+	
+	const FVector2D WindowSize(600, 400);
+
+	// Determine the position of the window so that it will spawn near the mouse, but not go off the screen.
+	FSlateRect AnchorRect;
+	if (Args.AnchorWidget.IsValid())
+	{
+		AnchorRect = Args.AnchorWidget->GetCachedGeometry().GetLayoutBoundingRect();
+	}
+	else
+	{
+		AnchorRect = FSlateRect::FromPointAndExtent(FSlateApplication::Get().GetCursorPos(), FVector2d());
+	}
+	AnchorRect = AnchorRect.ExtendBy(FMargin(20));
+	
+	const FVector2D ProposedPlacement = AnchorRect.GetTopLeft() + FVector2D(-WindowSize.X, -WindowSize.Y*0.2);
+	const FVector2D AdjustedSummonLocation = FSlateApplication::Get().CalculatePopupWindowPosition(AnchorRect, WindowSize, true, ProposedPlacement, Orient_Horizontal);
+
+	FText Title = Args.Title;
+	if (Title.IsEmpty())
+	{
+		Title = LOCTEXT("GameplayTagQueryWidget_Title", "Tag Query Editor");
+	}
+
+	TSharedRef<SGameplayTagQueryWidget> QueryWidget = SNew(SGameplayTagQueryWidget, Args.EditableQueries)
+			.OnQueriesCommitted(Args.OnQueriesCommitted)
+			.Filter(Args.Filter)
+			.ReadOnly(Args.bReadOnly);
+
+	TSharedRef<SWindow> Window = SNew(SWindow)
+		.Title(Args.Title)
+		.HasCloseButton(true)
+		.SupportsMinimize(false)
+		.SupportsMaximize(false)
+		.AutoCenter(EAutoCenter::None)
+		.ScreenPosition(AdjustedSummonLocation)
+		.ClientSize(WindowSize)
+		[
+			QueryWidget
+		];
+
+	// NOTE: FGlobalTabmanager::Get()-> is actually dereferencing a SharedReference, not a SharedPtr, so it cannot be null.
+	if (FGlobalTabmanager::Get()->GetRootWindow().IsValid())
+	{
+		FSlateApplication::Get().AddWindowAsNativeChild(Window, FGlobalTabmanager::Get()->GetRootWindow().ToSharedRef());
+	}
+	else
+	{
+		FSlateApplication::Get().AddWindow(Window);
+	}
+
+	GlobalTagQueryWidget = QueryWidget;
+	GlobalTagQueryWidgetWindow = Window;
+
+	auto CloseWindow = [WeakQueryWidget = GlobalTagQueryWidget]()
+	{
+		CloseGameplayTagQueryWindow(WeakQueryWidget);
+	};
+
+	QueryWidget->SetOnOk(FSimpleDelegate::CreateLambda(CloseWindow));
+	QueryWidget->SetOnCancel(FSimpleDelegate::CreateLambda(CloseWindow));
+
+	return QueryWidget;
+}
+
+void CloseGameplayTagQueryWindow(TWeakPtr<SGameplayTagQueryWidget> QueryWidget)
+{
+	if (GlobalTagQueryWidget.IsValid() && GlobalTagQueryWidgetWindow.IsValid())
+	{
+		if (!QueryWidget.IsValid() || QueryWidget == GlobalTagQueryWidget)
+		{
+			GlobalTagQueryWidgetWindow.Pin()->RequestDestroyWindow();
+		}
+	}
+	
+	GlobalTagQueryWidget = nullptr;
+	GlobalTagQueryWidgetWindow = nullptr;
+}
+
+} // UE::GameplayTags::Editor
+
 #undef LOCTEXT_NAMESPACE
diff --git a/Engine/Plugins/Editor/GameplayTagsEditor/Source/GameplayTagsEditor/Private/SGameplayTagQueryWidget.h b/Engine/Plugins/Editor/GameplayTagsEditor/Source/GameplayTagsEditor/Private/SGameplayTagQueryWidget.h
--- a/Engine/Plugins/Editor/GameplayTagsEditor/Source/GameplayTagsEditor/Private/SGameplayTagQueryWidget.h
+++ b/Engine/Plugins/Editor/GameplayTagsEditor/Source/GameplayTagsEditor/Private/SGameplayTagQueryWidget.h
@@ -6,97 +6,106 @@
 
 class IPropertyHandle;
 struct FGameplayTagQuery;
-
 class IDetailsView;
 struct FPropertyChangedEvent;
+class UEditableGameplayTagQuery;
 
 /** Widget allowing user to tag assets with gameplay tags */
 class SGameplayTagQueryWidget : public SCompoundWidget
 {
 public:
+
+	DECLARE_DELEGATE_OneParam( FOnQueriesCommitted, const TArray<FGameplayTagQuery>& /*TagQueries*/ )
+	
 	SLATE_BEGIN_ARGS( SGameplayTagQueryWidget )
-		: _ReadOnly(false), _AutoSave(false)
+		: _ReadOnly(false)
 	{}
-		SLATE_ARGUMENT(bool, ReadOnly)		// Flag to set if the list is read only
-		SLATE_ARGUMENT(bool, AutoSave)		// Flag to set if edits should be applied automatically (hides buttons)
-		SLATE_EVENT(FSimpleDelegate, OnClosePreSave) // Called when "Save and Close" button clicked
-		SLATE_EVENT(FSimpleDelegate, OnSaveAndClose) // Called when "Save and Close" button clicked
-		SLATE_EVENT(FSimpleDelegate, OnCancel) // Called when "Close Without Saving" button clicked
-		SLATE_EVENT(FSimpleDelegate, OnQueryChanged)	// Called when the user has modified the query
+		/** Flag to set if the list is read only */
+		SLATE_ARGUMENT(bool, ReadOnly)
+		/** Comma delimited string of tag root names to filter by */
+		SLATE_ARGUMENT(FString, Filter)
+		/** Delegate to call when dialog Ok is pressed. */
+		SLATE_EVENT(FSimpleDelegate, OnOk)
+		/** Delegate to call when dialog Cancel is pressed. */
+		SLATE_EVENT(FSimpleDelegate, OnCancel)
+		/** Delegate to call when queries are committed. */	
+		SLATE_EVENT(FOnQueriesCommitted, OnQueriesCommitted)
 	SLATE_END_ARGS()
 
-		/** Simple struct holding a tag query and its owner for generic re-use of the widget */
-	struct FEditableGameplayTagQueryDatum
-	{
-		/** Constructor */
-		FEditableGameplayTagQueryDatum(class UObject* InOwnerObj, struct FGameplayTagQuery* InTagQuery, FString* InTagExportText=nullptr)
-			: TagQueryOwner(InOwnerObj)
-			, TagQuery(InTagQuery)
-			, TagQueryExportText(InTagExportText)
-		{}
+	/** Construct the actual widget */
+	void Construct(const FArguments& InArgs, const TArray<FGameplayTagQuery>& InTagQueries);
 
-		/** Owning UObject of the query being edited */
-		TWeakObjectPtr<class UObject> TagQueryOwner;
+	void SetOnOk(FSimpleDelegate InOnOk) { OnOk = InOnOk; }
+	void SetOnCancel(FSimpleDelegate InOnCancel) { OnCancel = InOnCancel; }
+	
+	virtual ~SGameplayTagQueryWidget() override;
 
-		/** Tag query to edit */
-		struct FGameplayTagQuery* TagQuery;
+private:
 
-		/** The export text for FGameplayTagQuery, useful in some circumstances */
-		FString* TagQueryExportText;
-	};
+	void OnGetCategoriesMetaFromPropertyHandle(TSharedPtr<IPropertyHandle> PropertyHandle, FString& MetaString) const;
 
-	/** Construct the actual widget */
-	void Construct(const FArguments& InArgs, const TArray<FEditableGameplayTagQueryDatum>& EditableTagQueries, const TSharedPtr<IPropertyHandle> InTagQueryPropertyHandle);
+	/** Called when the user clicks the "Save and Close" button */
+	FReply OnOkClicked();
 
-	virtual ~SGameplayTagQueryWidget() override;
+	/** Called when the user clicks the "Close Without Saving" button */
+	FReply OnCancelClicked() const;
 
-private:
+	/** Controls visibility of the "Save and Close" button */
+	EVisibility GetSaveAndCloseButtonVisibility() const;
+	/** Controls visibility of the "Close Without Saving" button */
+	EVisibility GetCancelButtonVisibility() const;
 
-	void OnGetCategoriesMetaFromPropertyHandle(TSharedPtr<IPropertyHandle> PropertyHandle, FString& MetaString);
+	/** Converts editable representation to FGameplayTagQuery and calls OnQueriesCommitted. */
+	void SaveToTagQuery();
 
 	/* Flag to set if the list is read only*/
 	uint32 bReadOnly : 1;
 
-	/** If true, query will be written immediately on all changes. Otherwise, will only be written on user prompt (via the buttons). */
-	uint32 bAutoSave : 1;
-
 	/** Containers to modify */
-	TArray<FEditableGameplayTagQueryDatum> TagQueries;
-
-	/** Property that holds the tag query */
-	TSharedPtr<IPropertyHandle> TagQueryPropertyHandle;
-
-	/** Called when "Save and Close" is clicked before we save the data. */
-	FSimpleDelegate OnClosePreSave;
+	TArray<FGameplayTagQuery> TagQueries;
 
-	/** Called when "Save and Close" is clicked after we have saved the data. */
-	FSimpleDelegate OnSaveAndClose;
+	/** Called when values are committed. */
+	FOnQueriesCommitted OnQueriesCommitted;
 
-	/** Called when the user has modified the query */
-	FSimpleDelegate OnQueryChanged;
+	/** Called when "OK" is clicked. */
+	FSimpleDelegate OnOk;
 
-	/** Called when "Close Without Saving" is clicked */
+	/** Called when "Cancel" is clicked. */
 	FSimpleDelegate OnCancel;
 
+	/** Comma delimited string of tag root names to filter by */
+	FString Filter;
+	
 	/** Properties Tab */
-	TSharedPtr<class IDetailsView> Details;
+	TSharedPtr<IDetailsView> Details;
 
-	static class UEditableGameplayTagQuery* CreateEditableQuery(FGameplayTagQuery& Q);
-	TWeakObjectPtr<class UEditableGameplayTagQuery> EditableQuery;
+	static UEditableGameplayTagQuery* CreateEditableQuery(FGameplayTagQuery& Q);
+	TWeakObjectPtr<UEditableGameplayTagQuery> EditableQuery;
+};
 
-	/** Called when the user clicks the "Save and Close" button */
-	FReply OnSaveAndCloseClicked();
+struct FGameplayTagQueryWindowArgs
+{
+	/** Title of the query editor window. */
+	FText Title;
 
-	/** Called when the user clicks the "Close Without Saving" button */
-	FReply OnCancelClicked() const;
+	/** Flag to set if the list is read only */
+	bool bReadOnly = false;
 
-	void OnFinishedChangingProperties(const FPropertyChangedEvent& PropertyChangedEvent);
+	/** Delegate to call when the query is changed.  */
+	SGameplayTagQueryWidget::FOnQueriesCommitted OnQueriesCommitted;
 
-	/** Controls visibility of the "Save and Close" button */
-	EVisibility GetSaveAndCloseButtonVisibility() const;
-	/** Controls visibility of the "Close Without Saving" button */
-	EVisibility GetCancelButtonVisibility() const;
+	/** Comma delimited string of tag root names to filter by */
+	FString Filter;
+	
+	/** The array of queries to edit. */
+	TArray<FGameplayTagQuery> EditableQueries;
 
-	void SaveToTagQuery();
+	/** (Optional) Widget to use to position the query window. */
+	TSharedPtr<SWidget> AnchorWidget;
 };
 
+namespace UE::GameplayTags::Editor
+{
+	TWeakPtr<SGameplayTagQueryWidget> OpenGameplayTagQueryWindow(const FGameplayTagQueryWindowArgs& Args);
+	void CloseGameplayTagQueryWindow(TWeakPtr<SGameplayTagQueryWidget> QueryWidget);
+};
\ No newline at end of file
diff --git a/Engine/Plugins/Editor/GameplayTagsEditor/Source/GameplayTagsEditor/Private/SRenameGameplayTagDialog.cpp b/Engine/Plugins/Editor/GameplayTagsEditor/Source/GameplayTagsEditor/Private/SRenameGameplayTagDialog.cpp
--- a/Engine/Plugins/Editor/GameplayTagsEditor/Source/GameplayTagsEditor/Private/SRenameGameplayTagDialog.cpp
+++ b/Engine/Plugins/Editor/GameplayTagsEditor/Source/GameplayTagsEditor/Private/SRenameGameplayTagDialog.cpp
@@ -8,6 +8,7 @@
 #include "Widgets/Text/STextBlock.h"
 #include "Framework/Application/SlateApplication.h"
 #include "Widgets/SWindow.h"
+#include "Widgets/Layout/SBox.h"
 
 #define LOCTEXT_NAMESPACE "RenameGameplayTag"
 
@@ -20,7 +21,7 @@ void SRenameGameplayTagDialog::Construct(const FArguments& InArgs)
 
 	ChildSlot
 	[
-		SNew( SBorder )
+		SNew( SBox )
 		.Padding(FMargin(15))
 		[
 			SNew( SVerticalBox )
@@ -138,8 +139,8 @@ void SRenameGameplayTagDialog::RenameAndClose()
 {
 	IGameplayTagsEditorModule& Module = IGameplayTagsEditorModule::Get();
 
-	FString TagToRename = GameplayTagNode->GetCompleteTag().GetTagName().ToString();
-	FString NewTagName = NewTagNameTextBox->GetText().ToString();
+	const FString TagToRename = GameplayTagNode->GetCompleteTag().GetTagName().ToString();
+	const FString NewTagName = NewTagNameTextBox->GetText().ToString();
 
 	if (Module.RenameTagInINI(TagToRename, NewTagName))
 	{
diff --git a/Engine/Plugins/Editor/GameplayTagsEditor/Source/GameplayTagsEditor/Public/AssetTypeActions_GameplayTagAssetBase.h b/Engine/Plugins/Editor/GameplayTagsEditor/Source/GameplayTagsEditor/Public/AssetTypeActions_GameplayTagAssetBase.h
--- a/Engine/Plugins/Editor/GameplayTagsEditor/Source/GameplayTagsEditor/Public/AssetTypeActions_GameplayTagAssetBase.h
+++ b/Engine/Plugins/Editor/GameplayTagsEditor/Source/GameplayTagsEditor/Public/AssetTypeActions_GameplayTagAssetBase.h
@@ -4,6 +4,7 @@
 
 #include "AssetTypeActions_Base.h"
 
+struct FGameplayTagContainer;
 
 /** Base asset type actions for any classes with gameplay tagging */
 class GAMEPLAYTAGSEDITOR_API FAssetTypeActions_GameplayTagAssetBase : public FAssetTypeActions_Base
@@ -25,7 +26,7 @@ private:
 	 * 
 	 * @param TagAssets	Assets to open the editor with
 	 */
-	void OpenGameplayTagEditor(TArray<class UObject*> Objects, TArray<struct FGameplayTagContainer*> Containers);
+	void OpenGameplayTagEditor(TArray<UObject*> Objects, TArray<FGameplayTagContainer> Containers) const;
 
 	/** Name of the property of the owned gameplay tag container */
 	FName OwnedGameplayTagPropertyName;
diff --git a/Engine/Plugins/Editor/GameplayTagsEditor/Source/GameplayTagsEditor/Public/GameplayTagContainerCustomizationOptions.h b/Engine/Plugins/Editor/GameplayTagsEditor/Source/GameplayTagsEditor/Public/GameplayTagContainerCustomizationOptions.h
--- a/Engine/Plugins/Editor/GameplayTagsEditor/Source/GameplayTagsEditor/Public/GameplayTagContainerCustomizationOptions.h
+++ b/Engine/Plugins/Editor/GameplayTagsEditor/Source/GameplayTagsEditor/Public/GameplayTagContainerCustomizationOptions.h
@@ -4,7 +4,7 @@
 
 #include "HAL/Platform.h"
 
-struct GAMEPLAYTAGSEDITOR_API FGameplayTagContainerCustomizationOptions
+struct UE_DEPRECATED(5.3, "Options are not used anymore to customize the property customization behavior.") GAMEPLAYTAGSEDITOR_API FGameplayTagContainerCustomizationOptions
 {
 	// If true, any Gameplay Tag Widget created should not offer an 'Add Tag' option 
 	bool bForceHideAddTag = false;
diff --git a/Engine/Plugins/Editor/GameplayTagsEditor/Source/GameplayTagsEditor/Public/GameplayTagCustomizationOptions.h b/Engine/Plugins/Editor/GameplayTagsEditor/Source/GameplayTagsEditor/Public/GameplayTagCustomizationOptions.h
--- a/Engine/Plugins/Editor/GameplayTagsEditor/Source/GameplayTagsEditor/Public/GameplayTagCustomizationOptions.h
+++ b/Engine/Plugins/Editor/GameplayTagsEditor/Source/GameplayTagsEditor/Public/GameplayTagCustomizationOptions.h
@@ -4,7 +4,7 @@
 
 #include "HAL/Platform.h"
 
-struct GAMEPLAYTAGSEDITOR_API FGameplayTagCustomizationOptions
+struct UE_DEPRECATED(5.3, "Options are not used anymore to customize the property customization behavior.") GAMEPLAYTAGSEDITOR_API FGameplayTagCustomizationOptions
 {
 	// If true, any Gameplay Tag Widget created should not offer an 'Add Tag' option 
 	bool bForceHideAddTag = false;
diff --git a/Engine/Plugins/Editor/GameplayTagsEditor/Source/GameplayTagsEditor/Public/GameplayTagsEditorModule.h b/Engine/Plugins/Editor/GameplayTagsEditor/Source/GameplayTagsEditor/Public/GameplayTagsEditorModule.h
--- a/Engine/Plugins/Editor/GameplayTagsEditor/Source/GameplayTagsEditor/Public/GameplayTagsEditorModule.h
+++ b/Engine/Plugins/Editor/GameplayTagsEditor/Source/GameplayTagsEditor/Public/GameplayTagsEditorModule.h
@@ -83,6 +83,7 @@ public:
 struct GAMEPLAYTAGSEDITOR_API FGameplayTagCustomizationPublic
 {
 	static TSharedRef<IPropertyTypeCustomization> MakeInstance();
+	UE_DEPRECATED(5.3, "FGameplayTagCustomizationOptions are deprecated. The options are not used anymore, use the version without options.")
 	static TSharedRef<IPropertyTypeCustomization> MakeInstanceWithOptions(const FGameplayTagCustomizationOptions& Options);
 };
 
@@ -90,6 +91,7 @@ struct GAMEPLAYTAGSEDITOR_API FGameplayTagCustomizationPublic
 struct GAMEPLAYTAGSEDITOR_API FGameplayTagContainerCustomizationPublic
 {
 	static TSharedRef<IPropertyTypeCustomization> MakeInstance();
+	UE_DEPRECATED(5.3, "FGameplayTagContainerCustomizationOptions are deprecated. The options are not used anymore, use the version without options.")
 	static TSharedRef<IPropertyTypeCustomization> MakeInstanceWithOptions(const FGameplayTagContainerCustomizationOptions& Options);
 };
 
diff --git a/Engine/Plugins/Editor/GameplayTagsEditor/Source/GameplayTagsEditor/Public/SGameplayTagPicker.h b/Engine/Plugins/Editor/GameplayTagsEditor/Source/GameplayTagsEditor/Public/SGameplayTagPicker.h
new file mode 100644
--- /dev/null
+++ b/Engine/Plugins/Editor/GameplayTagsEditor/Source/GameplayTagsEditor/Public/SGameplayTagPicker.h
@@ -0,0 +1,468 @@
+﻿// Copyright Epic Games, Inc. All Rights Reserved.
+
+#pragma once
+
+#include "SlateFwd.h"
+#include "UObject/Object.h"
+#include "Layout/Visibility.h"
+#include "Widgets/DeclarativeSyntaxSupport.h"
+#include "Input/Reply.h"
+#include "Widgets/SWidget.h"
+#include "Widgets/SCompoundWidget.h"
+#include "Widgets/Views/STableViewBase.h"
+#include "Widgets/Views/STableRow.h"
+#include "Widgets/Views/STreeView.h"
+#include "GameplayTagsManager.h"
+#include "EditorUndoClient.h"
+
+class IPropertyHandle;
+class SComboButton;
+
+/** Determines the behavior of the gameplay tag UI depending on where it's used */
+enum class EGameplayTagPickerMode : uint8
+{
+	SelectionMode,		// Allows selecting tags for gameplay tag containers.
+	ManagementMode,		// Allows renaming/deleting tags, does not require gameplay tag containers.
+	HybridMode			// Allows doing operations from both Selection and Management mode.
+};
+
+/** Widget allowing user to tag assets with gameplay tags */
+class GAMEPLAYTAGSEDITOR_API SGameplayTagPicker : public SCompoundWidget, public FSelfRegisteringEditorUndoClient
+{
+public:
+
+	/** Called when a tag status is changed */
+	DECLARE_DELEGATE_OneParam(FOnTagChanged, const TArray<FGameplayTagContainer>& /*TagContainers*/)
+
+	/** Called on when tags might need refreshing (e.g. after undo/redo or when tags change). Use SetTagContainers() to set the new tags. */
+	DECLARE_DELEGATE_OneParam(FOnRefreshTagContainers, SGameplayTagPicker& /*TagPicker*/)
+	
+	enum class ETagFilterResult
+	{
+		IncludeTag,
+		ExcludeTag
+	};
+
+	DECLARE_DELEGATE_RetVal_OneParam(ETagFilterResult, FOnFilterTag, const TSharedPtr<FGameplayTagNode>&)
+
+	SLATE_BEGIN_ARGS(SGameplayTagPicker)
+		: _Filter()
+		, _SettingsName(TEXT(""))
+		, _ReadOnly(false)
+		, _MultiSelect(true)
+		, _RestrictedTags(false)
+		, _ShowMenuItems(false)
+		, _PropertyHandle(nullptr)
+		, _GameplayTagPickerMode(EGameplayTagPickerMode::SelectionMode)
+		, _MaxHeight(260.0f)
+		, _Padding(FMargin(2.0f))
+	{}
+		// Comma delimited string of tag root names to filter by
+		SLATE_ARGUMENT(FString, Filter)
+
+		// Optional filter function called when generating the tag list
+		SLATE_EVENT(FOnFilterTag, OnFilterTag)
+
+		// The name that will be used for the settings file
+		SLATE_ARGUMENT(FString, SettingsName)
+
+		// Flag to set if the list is read only
+		SLATE_ARGUMENT(bool, ReadOnly)
+
+		// If we can select multiple entries
+		SLATE_ARGUMENT(bool, MultiSelect)
+
+		// If we are dealing with restricted tags or regular gameplay tags
+		SLATE_ARGUMENT(bool, RestrictedTags)
+
+		// If set, wraps the picker in a menu builder and adds common menu commands.
+		SLATE_ARGUMENT(bool, ShowMenuItems)  
+
+		// Property handle to FGameplayTag or FGameplayTagContainer, used to get and modify the edited tags.
+		SLATE_ARGUMENT(TSharedPtr<IPropertyHandle>, PropertyHandle)
+
+		// Tags or tag containers to modify. If MultiSelect is false, the container will container single tags.
+		// If PropertyHandle is set, the tag containers will be ignored.
+		SLATE_ARGUMENT(TArray<FGameplayTagContainer>, TagContainers)
+
+		// Called when a tag status changes
+		SLATE_EVENT(FOnTagChanged, OnTagChanged)
+
+		// Determines behavior of the menu based on where it's used
+		SLATE_ARGUMENT(EGameplayTagPickerMode, GameplayTagPickerMode)
+
+		// Caps the height of the gameplay tag tree, 0.0 means uncapped.
+		SLATE_ARGUMENT(float, MaxHeight)
+
+		// Padding inside the picker widget
+		SLATE_ARGUMENT(FMargin, Padding)
+
+		// Called on when tags might need refreshing (e.g. after undo/redo or when tags change).
+		SLATE_EVENT(FOnRefreshTagContainers, OnRefreshTagContainers)
+	SLATE_END_ARGS()
+
+	/**
+	 * Given a property handle, try and enumerate the tag containers from within it (when dealing with a struct property of type FGameplayTag or FGameplayTagContainer).
+	 * @return True if it was possible to enumerate containers (even if no containers were enumerated), or false otherwise.
+	 */
+	static bool EnumerateEditableTagContainersFromPropertyHandle(const TSharedRef<IPropertyHandle>& PropHandle, TFunctionRef<bool(const FGameplayTagContainer&)> Callback);
+
+	/**
+	 * Given a property handle, try and extract the editable tag containers from within it (when dealing with a struct property of type FGameplayTagContainer).
+	 * @return True if it was possible to extract containers (even if no containers were extracted), or false otherwise.
+	 */
+	static bool GetEditableTagContainersFromPropertyHandle(const TSharedRef<IPropertyHandle>& PropHandle, TArray<FGameplayTagContainer>& OutEditableContainers);
+
+	/** Construct the actual widget */
+	void Construct(const FArguments& InArgs);
+	
+	virtual void Tick(const FGeometry& AllottedGeometry, const double InCurrentTime, const float InDeltaTime) override;
+
+	/** Updates the tag list when the filter text changes */
+	void OnFilterTextChanged(const FText& InFilterText);
+
+	/** Returns true if this TagNode has any children that match the current filter */
+	bool FilterChildrenCheck(TSharedPtr<FGameplayTagNode>& InItem) const;
+
+	/** Child recursion helper for FilterChildrenCheck() */
+	bool FilterChildrenCheckRecursive(TSharedPtr<FGameplayTagNode>& InItem) const;
+
+	/** Refreshes the tags that should be displayed by the widget */
+	void RefreshTags();
+
+	/** Forces the widget to refresh its tags on the next tick */
+	void RefreshOnNextTick();
+
+	/** Scrolls the view to specified tag. */
+	void RequestScrollToView(const FGameplayTag Tag);
+	
+	/** Gets the widget to focus once the menu opens. */
+	TSharedPtr<SWidget> GetWidgetToFocusOnOpen();
+
+	/** Sets tag containers to edit. */
+	void SetTagContainers(TConstArrayView<FGameplayTagContainer> TagContainers);
+	
+private:
+
+	enum class EGameplayTagAdd : uint8
+	{
+		Root,
+		Child,
+		Duplicate,
+	};
+
+	// FSelfRegisteringEditorUndoClient
+	virtual void PostUndo(bool bSuccess) override;
+	virtual void PostRedo(bool bSuccess) override;
+	// ~FSelfRegisteringEditorUndoClient
+
+	/** Verify the tags are all valid and if not prompt the user. */
+	void VerifyAssetTagValidity();
+
+	/* Filters the tree view based on the current filter text. */
+	void FilterTagTree();
+
+	/* string that sets the section of the ini file to use for this class*/ 
+	static const FString SettingsIniSection;
+
+	/* Filename of ini file to store state of the UI */
+	static const FString& GetGameplayTagsEditorStateIni();
+
+	/**
+	 * Generate a row widget for the specified item node and table
+	 * 
+	 * @param InItem		Tag node to generate a row widget for
+	 * @param OwnerTable	Table that owns the row
+	 * 
+	 * @return Generated row widget for the item node
+	 */
+	TSharedRef<ITableRow> OnGenerateRow(TSharedPtr<FGameplayTagNode> InItem, const TSharedRef<STableViewBase>& OwnerTable);
+
+	/**
+	 * Get children nodes of the specified node
+	 * 
+	 * @param InItem		Node to get children of
+	 * @param OutChildren	[OUT] Array of children nodes, if any
+	 */
+	void OnGetChildren(TSharedPtr<FGameplayTagNode> InItem, TArray< TSharedPtr<FGameplayTagNode> >& OutChildren);
+
+	/**
+	 * Called via delegate when the status of a check box in a row changes
+	 * 
+	 * @param NewCheckState	New check box state
+	 * @param NodeChanged	Node that was checked/unchecked
+	 */
+	void OnTagCheckStatusChanged(ECheckBoxState NewCheckState, TSharedPtr<FGameplayTagNode> NodeChanged);
+
+	/**
+	 * Called via delegate to determine the checkbox state of the specified node
+	 * 
+	 * @param Node	Node to find the checkbox state of
+	 * 
+	 * @return Checkbox state of the specified node
+	 */
+	ECheckBoxState IsTagChecked(TSharedPtr<FGameplayTagNode> Node) const;
+
+	/**
+	 * @return true if the exact Tag provided is included in any of the tag containers the widget is editing.
+	 */
+	bool IsExactTagInCollection(TSharedPtr<FGameplayTagNode> Node) const;
+
+	/**
+	 * Called via delegate when the status of the allow non-restricted children check box in a row changes
+	 *
+	 * @param NewCheckState	New check box state
+	 * @param NodeChanged	Node that was checked/unchecked
+	 */
+	void OnAllowChildrenTagCheckStatusChanged(ECheckBoxState NewCheckState, TSharedPtr<FGameplayTagNode> NodeChanged);
+
+	/**
+	 * Called via delegate to determine the non-restricted children checkbox state of the specified node
+	 *
+	 * @param Node	Node to find the non-restricted children checkbox state of
+	 *
+	 * @return Non-restricted children heckbox state of the specified node
+	 */
+	ECheckBoxState IsAllowChildrenTagChecked(TSharedPtr<FGameplayTagNode> Node) const;
+
+	/** Helper function to determine the visibility of the checkbox for allowing non-restricted children of restricted gameplay tags */
+	EVisibility DetermineAllowChildrenVisible(TSharedPtr<FGameplayTagNode> Node) const;
+
+	/**
+	 * Helper function called when the specified node is checked
+	 * 
+	 * @param NodeChecked	Node that was checked by the user
+	 */
+	void OnTagChecked(TSharedPtr<FGameplayTagNode> NodeChecked);
+
+	/**
+	 * Helper function called when the specified node is unchecked
+	 * 
+	 * @param NodeUnchecked	Node that was unchecked by the user
+	 */
+	void OnTagUnchecked(TSharedPtr<FGameplayTagNode> NodeUnchecked);
+
+	/**
+	 * Recursive function to uncheck all child tags
+	 * 
+	 * @param NodeUnchecked	Node that was unchecked by the user
+	 * @param EditableContainer The container we are removing the tags from
+	 */
+	void UncheckChildren(TSharedPtr<FGameplayTagNode> NodeUnchecked, FGameplayTagContainer& EditableContainer);
+
+	/**
+	 * Called via delegate to determine the text colour of the specified node
+	 *
+	 * @param Node	Node to find the colour of
+	 *
+	 * @return Text colour of the specified node
+	 */
+	FSlateColor GetTagTextColour(TSharedPtr<FGameplayTagNode> Node) const;
+
+	/** Called when the user clicks the "Create and Manage Gameplay Tags" button; Opens Gameplay Tag manager window. */
+	void OnManageTagsClicked(TSharedPtr<SComboButton> OwnerCombo);
+
+	/** Called when the user clicks the "Clear All" button; Clears all tags */
+	void OnClearAllClicked(TSharedPtr<SComboButton> OwnerCombo);
+
+	/** Called when the user clicks the "Expand All" button; Expands the entire tag tree */
+	void OnExpandAllClicked(TSharedPtr<SComboButton> OwnerCombo);
+
+	/** Called when the user clicks the "Collapse All" button; Collapses the entire tag tree */
+	void OnCollapseAllClicked(TSharedPtr<SComboButton> OwnerCombo);
+
+	/**
+	 * Helper function to set the expansion state of the tree widget
+	 * 
+	 * @param bExpand If true, expand the entire tree; Otherwise, collapse the entire tree
+	 * @param bPersistExpansion If true, persist the expansion state.
+	 */
+	void SetTagTreeItemExpansion(bool bExpand, bool bPersistExpansion = false);
+
+	/**
+	 * Helper function to set the expansion state of a specific node
+	 * 
+	 * @param Node		Node to set the expansion state of
+	 * @param bExpand	If true, expand the node; Otherwise, collapse the node
+	 * @param bPersistExpansion If true, persist the expansion state.
+	 */
+	void SetTagNodeItemExpansion(TSharedPtr<FGameplayTagNode> Node, bool bExpand, bool bPersistExpansion = false);
+
+	/** Load settings for the tags*/
+	void LoadSettings();
+
+	/** Migrate settings */
+	void MigrateSettings();
+
+	/** Helper function to determine the visibility of the expandable UI controls */
+	EVisibility DetermineExpandableUIVisibility() const;
+
+	/** Helper function to determine the visibility of the Add New Tag widget */
+	bool CanAddNewTag() const;
+
+	/** Helper function to determine the visibility of the Add New Subtag widget */
+	bool CanAddNewSubTag(TSharedPtr<FGameplayTagNode> Node) const;
+
+	/** Helper function to determine if tag can be modified */
+	bool CanModifyTag(TSharedPtr<FGameplayTagNode> Node) const;
+
+	/** Recursive load function to go through all tags in the tree and set the expansion*/
+	void LoadTagNodeItemExpansion(TSharedPtr<FGameplayTagNode> Node );
+
+	/** Recursive function to go through all tags in the tree and set the expansion to default*/
+	void SetDefaultTagNodeItemExpansion(TSharedPtr<FGameplayTagNode> Node);
+
+	/** Expansion changed callback */
+	void OnExpansionChanged(TSharedPtr<FGameplayTagNode> InItem, bool bIsExpanded);
+
+	/** Callback for when a new tag is added */
+	void OnGameplayTagAdded(const FString& TagName, const FString& TagComment, const FName& TagSource);
+
+	/** Callback when the user wants to add a subtag to an existing tag */
+	FReply OnAddSubtagClicked(TSharedPtr<FGameplayTagNode> InTagNode);
+
+	void AddTagDialog(const EGameplayTagAdd Mode, TSharedPtr<FGameplayTagNode> InTagNode = TSharedPtr<FGameplayTagNode>());
+
+	FReply OnAddRootTagClicked();
+
+	FText GetHighlightText() const;
+	
+	/** Creates a dropdown menu to provide additional functionality for tags (renaming, deletion, search for references, etc.) */
+	TSharedRef<SWidget> MakeTagActionsMenu(TSharedPtr<FGameplayTagNode> InTagNode, TSharedPtr<SComboButton> OwnerCombo, const bool bInShouldCloseWindowAfterMenuSelection);
+
+	/** Creates add menu content. */
+	TSharedRef<SWidget> MakeAddMenu();
+
+	/** Creates actions menu content. */
+	TSharedRef<SWidget> MakeActionsMenu(TSharedPtr<SComboButton> OwnerCombo);
+
+	/** Creates settings menu content. */
+	TSharedRef<SWidget> MakeSettingsMenu(TSharedPtr<SComboButton> OwnerCombo);
+
+	void OnAddSubTag(TSharedPtr<FGameplayTagNode> InTagNode, TSharedPtr<SComboButton> OwnerCombo);
+
+	void OnDuplicateTag(TSharedPtr<FGameplayTagNode> InTagNode, TSharedPtr<SComboButton> OwnerCombo);
+	
+	/** Attempts to rename the tag through a dialog box */
+	void OnRenameTag(TSharedPtr<FGameplayTagNode> InTagNode, TSharedPtr<SComboButton> OwnerCombo);
+
+	/** Attempts to delete the specified tag */
+	void OnDeleteTag(TSharedPtr<FGameplayTagNode> InTagNode, TSharedPtr<SComboButton> OwnerCombo);
+
+	/** Attempts to select the exact specified tag*/
+	void OnSelectExactTag(TSharedPtr<FGameplayTagNode> InTagNode, TSharedPtr<SComboButton> OwnerCombo);
+
+	/** Attempts to unselect the specified tag, but not the children */
+	void OnUnselectExactTag(TSharedPtr<FGameplayTagNode> InTagNode, TSharedPtr<SComboButton> OwnerCombo);
+
+	/** Searches for all references for the selected tag */
+	void OnSearchForReferences(TSharedPtr<FGameplayTagNode> InTagNode, TSharedPtr<SComboButton> OwnerCombo);
+
+	void OnCopyTagNameToClipboard(TSharedPtr<FGameplayTagNode> InTagNode, TSharedPtr<SComboButton> OwnerCombo);
+	
+	/** Returns true if the user can select tags from the widget */
+	bool CanSelectTags() const;
+
+	void OnContainersChanged();
+
+	/** Opens a dialog window to rename the selected tag */
+	void OpenRenameGameplayTagDialog(TSharedPtr<FGameplayTagNode> GameplayTagNode) const;
+
+	/** Delegate that is fired when a tag is successfully renamed */
+	void OnGameplayTagRenamed(FString OldTagName, FString NewTagName);
+
+	/** Populate tag items from the gameplay tags manager. */
+	void GetFilteredGameplayRootTags(const FString& InFilterString, TArray<TSharedPtr<FGameplayTagNode>>& OutNodes) const;
+
+	/** Called to create context menu for the tag tree items. */
+	TSharedPtr<SWidget> OnTreeContextMenuOpening();
+
+	/** Called when the tag tree selection changes. */
+	void OnTreeSelectionChanged(TSharedPtr<FGameplayTagNode> SelectedItem, ESelectInfo::Type SelectInfo);
+
+	/** Called to handle key presses in the tag tree. */
+	FReply OnTreeKeyDown(const FGeometry& Geometry, const FKeyEvent& Key);
+	
+	/* Holds the Name of this TagContainer used for saving out expansion settings */
+	FString SettingsName;
+
+	/* Filter string used during search box */
+	FString FilterString;
+
+	/** root filter (passed in on creation) */
+	FString RootFilterString;
+
+	/** User specified filter function. */
+	FOnFilterTag TagFilter; 
+
+	/* Flag to set if the list is read only*/
+	bool bReadOnly = false;
+
+	/* Flag to set if we can select multiple items form the list*/
+	bool bMultiSelect = true;
+
+	/** If true, refreshes tags on the next frame */
+	bool bDelayRefresh = false;
+
+	/** If true, this widget is displaying restricted tags; if false this widget displays regular gameplay tags. */
+	bool bRestrictedTags = false;
+
+	/** The maximum height of the gameplay tag tree. If 0, the height is unbound. */
+	float MaxHeight = 260.0f;
+
+	/* Array of all tags */
+	TArray<TSharedPtr<FGameplayTagNode>> TagItems;
+
+	/* Array of tags filtered to be displayed in the TreeView */
+	TArray<TSharedPtr<FGameplayTagNode>> FilteredTagItems;
+
+	/** Container widget holding the tag tree */
+	TSharedPtr<SBorder> TagTreeContainerWidget;
+
+	/** Tree widget showing the gameplay tag library */
+	TSharedPtr<STreeView<TSharedPtr<FGameplayTagNode>>> TagTreeWidget;
+
+	/** Allows for the user to find a specific gameplay tag in the tree */
+	TSharedPtr<SSearchBox> SearchTagBox;
+
+	/** Property to edit */
+	TSharedPtr<IPropertyHandle> PropertyHandle;
+
+	/** Containers to modify, ignored if PropertyHandle is set. */
+	TArray<FGameplayTagContainer> TagContainers;
+
+	/** Called when the Tag list changes*/
+	FOnTagChanged OnTagChanged;
+
+	/** Called on when tags might need refreshing (e.g. after undo/redo or when tags change). */
+	FOnRefreshTagContainers OnRefreshTagContainers;
+
+	/** Tag to scroll to, cleared after scroll is requested. */
+	FGameplayTag RequestedScrollToTag;
+	
+	/** Determines behavior of the widget */
+	EGameplayTagPickerMode GameplayTagPickerMode = EGameplayTagPickerMode::SelectionMode;
+
+	/** Guard value used to prevent feedback loops in selection handling. */
+	bool bInSelectionChanged = false;
+
+	/** Guard value used define if expansion code should persist the expansion operations. */
+	bool bPersistExpansionChange = true;
+
+	/** Expanded items cached from settings */
+	TSet<TSharedPtr<FGameplayTagNode>> CachedExpandedItems;
+};
+
+
+struct FGameplayTagManagerWindowArgs
+{
+	FText Title;
+	FString Filter; // Comma delimited string of tag root names to filter by
+	bool bRestrictedTags;
+};
+
+namespace UE::GameplayTags::Editor
+{
+	GAMEPLAYTAGSEDITOR_API TWeakPtr<SGameplayTagPicker> OpenGameplayTagManager(const FGameplayTagManagerWindowArgs& Args);
+	GAMEPLAYTAGSEDITOR_API TSharedRef<SWidget> Create(const FGameplayTagManagerWindowArgs& Args);
+};
\ No newline at end of file
diff --git a/Engine/Source/Runtime/GameplayTags/Classes/GameplayTagsManager.h b/Engine/Source/Runtime/GameplayTags/Classes/GameplayTagsManager.h
--- a/Engine/Source/Runtime/GameplayTags/Classes/GameplayTagsManager.h
+++ b/Engine/Source/Runtime/GameplayTags/Classes/GameplayTagsManager.h
@@ -294,6 +294,7 @@ private:
 
 	friend class UGameplayTagsManager;
 	friend class SGameplayTagWidget;
+	friend class SGameplayTagPicker;
 };
 
 /** Holds data about the tag dictionary, is in a singleton UObject */
diff --git a/Engine/Source/Runtime/GameplayTags/Classes/GameplayTagsSettings.h b/Engine/Source/Runtime/GameplayTags/Classes/GameplayTagsSettings.h
--- a/Engine/Source/Runtime/GameplayTags/Classes/GameplayTagsSettings.h
+++ b/Engine/Source/Runtime/GameplayTags/Classes/GameplayTagsSettings.h
@@ -155,14 +155,19 @@ class GAMEPLAYTAGS_API UGameplayTagsSettings : public UGameplayTagsList
 	/** A list of .ini files used to store restricted gameplay tags. */
 	UPROPERTY(config, EditAnywhere, AdvancedDisplay, Category = "Advanced Gameplay Tags")
 	TArray<FRestrictedConfigInfo> RestrictedConfigFiles;
+
 #if WITH_EDITORONLY_DATA
-	// Dummy parameter used to hook the editor UI
+	// Dummy parameters used to hook the editor UI
 	/** Restricted Gameplay Tags.
 	 * 
 	 *  Restricted tags are intended to be top level tags that are important for your data hierarchy and modified by very few people.
 	 */
 	UPROPERTY(EditAnywhere, AdvancedDisplay, transient, Category = "Advanced Gameplay Tags")
 	FString RestrictedTagList;
+
+	/** Add a new gameplay tag config file for saving plugin or game-specific tags. */
+	UPROPERTY(EditAnywhere, transient, Category = "GameplayTags")
+	FString NewTagSource;
 #endif
 
 #if WITH_EDITOR
