Gameplay tags: Fix crash trying to access nullptr from property handle raw data.

diff --git a/Engine/Plugins/Editor/GameplayTagsEditor/Source/GameplayTagsEditor/Private/SGameplayTagPicker.cpp b/Engine/Plugins/Editor/GameplayTagsEditor/Source/GameplayTagsEditor/Private/SGameplayTagPicker.cpp
--- a/Engine/Plugins/Editor/GameplayTagsEditor/Source/GameplayTagsEditor/Private/SGameplayTagPicker.cpp
+++ b/Engine/Plugins/Editor/GameplayTagsEditor/Source/GameplayTagsEditor/Private/SGameplayTagPicker.cpp
@@ -13,11 +13,7 @@
 #include "Textures/SlateIcon.h"
 #include "PropertyHandle.h"
 #include "Widgets/Input/SSearchBox.h"
-#include "Widgets/Text/SInlineEditableTextBlock.h"
 #include "GameplayTagsEditorModule.h"
-#include "AssetToolsModule.h"
-#include "Widgets/Input/SHyperlink.h"
-#include "Framework/Notifications/NotificationManager.h"
 #include "Framework/Application/SlateApplication.h"
 #include "SAddNewGameplayTagWidget.h"
 #include "SAddNewRestrictedGameplayTagWidget.h"
@@ -26,7 +22,6 @@
 #include "Editor.h"
 #include "Framework/Commands/UIAction.h"
 #include "Framework/MultiBox/MultiBoxBuilder.h"
-#include "GameplayTagStyle.h"
 #include "Framework/Docking/TabManager.h"
 #include "HAL/PlatformApplicationMisc.h"
 #include "SPositiveActionButton.h"
@@ -44,7 +39,12 @@ bool SGameplayTagPicker::EnumerateEditableTagContainersFromPropertyHandle(const
 	{
 		PropHandle->EnumerateRawData([&Callback](void* RawData, const int32 /*DataIndex*/, const int32 /*NumDatas*/)
 		{
-			return Callback(*static_cast<FGameplayTagContainer*>(RawData));
+			// Report empty container even if the instance data is null to match all the instances indices of the property handle.
+			if (RawData)
+			{
+				return Callback(*static_cast<FGameplayTagContainer*>(RawData));
+			}
+			return Callback(FGameplayTagContainer());
 		});
 		return true;
 	}
@@ -52,7 +52,12 @@ bool SGameplayTagPicker::EnumerateEditableTagContainersFromPropertyHandle(const
 	{
 		PropHandle->EnumerateRawData([&Callback](void* RawData, const int32 /*DataIndex*/, const int32 /*NumDatas*/)
 		{
-			const FGameplayTagContainer Container(*static_cast<FGameplayTag*>(RawData));
+			// Report empty container even if the instance data is null to match all the instances indices of the property handle.
+			FGameplayTagContainer Container;
+			if (RawData)
+			{
+				Container.AddTag(*static_cast<FGameplayTag*>(RawData));
+			}
 			return Callback(Container);
 		});
 		return true;
diff --git a/Engine/Plugins/Editor/GameplayTagsEditor/Source/GameplayTagsEditor/Private/SGameplayTagQueryEntryBox.cpp b/Engine/Plugins/Editor/GameplayTagsEditor/Source/GameplayTagsEditor/Private/SGameplayTagQueryEntryBox.cpp
--- a/Engine/Plugins/Editor/GameplayTagsEditor/Source/GameplayTagsEditor/Private/SGameplayTagQueryEntryBox.cpp
+++ b/Engine/Plugins/Editor/GameplayTagsEditor/Source/GameplayTagsEditor/Private/SGameplayTagQueryEntryBox.cpp
@@ -5,7 +5,6 @@
 #include "Widgets/Text/STextBlock.h"
 #include "Widgets/Input/SButton.h"
 #include "Widgets/Images/SImage.h"
-#include "GameplayTagStyle.h"
 #include "ScopedTransaction.h"
 #include "Editor.h"
 #include "GameplayTagEditorUtilities.h"
@@ -238,13 +237,17 @@ void SGameplayTagQueryEntryBox::CacheQueryList()
 
 	if (PropertyHandle.IsValid())
 	{
+		// Cache queries from the property handle. Add empty queries even if the instance data is null so that the indices match with the property handle.
 		TArray<void*> RawStructData;
 		PropertyHandle->AccessRawData(RawStructData);
 		
 		for (int32 Idx = 0; Idx < RawStructData.Num(); ++Idx)
 		{
-			const FGameplayTagQuery& Query = *(FGameplayTagQuery*)RawStructData[Idx]; 
-			CachedQueries.Add(Query);
+			FGameplayTagQuery& Query = CachedQueries.AddDefaulted_GetRef();
+			if (RawStructData[Idx])
+			{
+				Query = *(FGameplayTagQuery*)RawStructData[Idx]; 
+			}
 		}
 	}
 	else
