Fix localization key confict on SGameplayTagPicker

#rnx

diff --git a/Engine/Plugins/Editor/GameplayTagsEditor/Source/GameplayTagsEditor/Private/SGameplayTagPicker.cpp b/Engine/Plugins/Editor/GameplayTagsEditor/Source/GameplayTagsEditor/Private/SGameplayTagPicker.cpp
--- a/Engine/Plugins/Editor/GameplayTagsEditor/Source/GameplayTagsEditor/Private/SGameplayTagPicker.cpp
+++ b/Engine/Plugins/Editor/GameplayTagsEditor/Source/GameplayTagsEditor/Private/SGameplayTagPicker.cpp
@@ -377,7 +377,7 @@ TSharedRef<SWidget> SGameplayTagPicker::MakeActionsMenu(TSharedPtr<SComboButton>
 	FMenuBuilder MenuBuilder(/*bInShouldCloseWindowAfterMenuSelection*/false, nullptr);
 
 	MenuBuilder.AddMenuEntry(
-		LOCTEXT("GameplayTagPicker_ManageTags", "Create and Manage Gameplay Tags..."), FText::GetEmpty(), FSlateIcon(FAppStyle::GetAppStyleSetName(), "Icons.Settings"),
+		LOCTEXT("GameplayTagPicker_CreateAndManageTags", "Create and Manage Gameplay Tags..."), FText::GetEmpty(), FSlateIcon(FAppStyle::GetAppStyleSetName(), "Icons.Settings"),
 		FUIAction(FExecuteAction::CreateRaw(this, &SGameplayTagPicker::OnManageTagsClicked, OwnerCombo))
 	);
 	
