GameplayTags: Restored quick add tag workflow, improved contrast on tags
- Restored the quick add tag workflow
- Added menu option to go to a specific tag in manager
- Cleaned up the layouts of the different dialogs
- Reuse already open GameplayTag manager
- Added Tag highlight option for opening the manager window

diff --git a/Engine/Plugins/Editor/GameplayTagsEditor/Source/GameplayTagsEditor/Internal/SAddNewGameplayTagWidget.h b/Engine/Plugins/Editor/GameplayTagsEditor/Source/GameplayTagsEditor/Internal/SAddNewGameplayTagWidget.h
--- a/Engine/Plugins/Editor/GameplayTagsEditor/Source/GameplayTagsEditor/Internal/SAddNewGameplayTagWidget.h
+++ b/Engine/Plugins/Editor/GameplayTagsEditor/Source/GameplayTagsEditor/Internal/SAddNewGameplayTagWidget.h
@@ -26,10 +26,14 @@ public:
 
 	SLATE_BEGIN_ARGS(SAddNewGameplayTagWidget)
 		: _NewTagName(TEXT(""))
+		, _Padding(FMargin(15))
+		, _AddButtonPadding(FMargin(0, 16, 0, 0))
 	{}
 		SLATE_EVENT(FOnGameplayTagAdded, OnGameplayTagAdded)	// Callback for when a new tag is added	
 		SLATE_EVENT(FIsValidTag, IsValidTag)
 		SLATE_ARGUMENT(FString, NewTagName) // String that will initially populate the New Tag Name field
+		SLATE_ARGUMENT(FMargin, Padding) // Padding around the widget
+		SLATE_ARGUMENT(FMargin, AddButtonPadding) // padding around the Add button
 	SLATE_END_ARGS();
 
 	GAMEPLAYTAGSEDITOR_API virtual ~SAddNewGameplayTagWidget() override;
diff --git a/Engine/Plugins/Editor/GameplayTagsEditor/Source/GameplayTagsEditor/Private/GameplayTagStyle.cpp b/Engine/Plugins/Editor/GameplayTagsEditor/Source/GameplayTagsEditor/Private/GameplayTagStyle.cpp
--- a/Engine/Plugins/Editor/GameplayTagsEditor/Source/GameplayTagsEditor/Private/GameplayTagStyle.cpp
+++ b/Engine/Plugins/Editor/GameplayTagsEditor/Source/GameplayTagsEditor/Private/GameplayTagStyle.cpp
@@ -32,11 +32,11 @@ public:
 		// Tag Chip
 		const FLinearColor ChipColor = FStyleColors::Hover.GetSpecifiedColor();
 		const FLinearColor ChipColorHover = FStyleColors::Hover2.GetSpecifiedColor();
-		const FLinearColor ChipColorDisable = ChipColor.CopyWithNewOpacity(0.5);
+		const FLinearColor ChipColorDisable = ChipColor.CopyWithNewOpacity(0.35);
 		
 		const FLinearColor HollowChipColor = ChipColor.CopyWithNewOpacity(0.0);
 		const FLinearColor HollowChipColorHover = ChipColorHover.CopyWithNewOpacity(0.15);
-		const FLinearColor HollowChipColorDisable = HollowChipColor.CopyWithNewOpacity(0.3);
+		const FLinearColor HollowChipColorDisable = HollowChipColor.CopyWithNewOpacity(0.0);
 
 		FButtonStyle ChipButtonSelected = FAppStyle::Get().GetWidgetStyle<FButtonStyle>("Button");
 		ChipButtonSelected
@@ -44,10 +44,10 @@ public:
 			.SetHovered(FSlateRoundedBoxBrush(ChipColorHover, 3.0f))
 			.SetPressed(FSlateRoundedBoxBrush(ChipColorHover, 3.0f))
 			.SetDisabled(FSlateRoundedBoxBrush(ChipColorDisable, 3.0f))
-			.SetNormalForeground(FStyleColors::ForegroundHeader)
-			.SetHoveredForeground(FStyleColors::ForegroundHover)
-			.SetPressedForeground(FStyleColors::ForegroundHover)
-			.SetDisabledForeground(FStyleColors::Foreground)
+			.SetNormalForeground(FStyleColors::ForegroundHover)
+			.SetHoveredForeground(FStyleColors::White)
+			.SetPressedForeground(FStyleColors::White)
+			.SetDisabledForeground(FStyleColors::ForegroundHover)
 			.SetNormalPadding(FMargin(5,2,2,2))
 			.SetPressedPadding(FMargin(5,3,2,1));
 		Set("GameplayTags.ChipButton.Selected", ChipButtonSelected);
@@ -58,10 +58,10 @@ public:
 			.SetHovered(FSlateRoundedBoxBrush(HollowChipColorHover, 3.0f, ChipColorHover, 1.0f))
 			.SetPressed(FSlateRoundedBoxBrush(HollowChipColorHover, 3.0f, ChipColorHover, 1.0f))
 			.SetDisabled(FSlateRoundedBoxBrush(HollowChipColorDisable, 3.0f, ChipColorDisable, 1.0f))
-			.SetNormalForeground(FStyleColors::ForegroundHeader)
-			.SetHoveredForeground(FStyleColors::ForegroundHover)
-			.SetPressedForeground(FStyleColors::ForegroundHover)
-			.SetDisabledForeground(FStyleColors::Foreground)
+			.SetNormalForeground(FStyleColors::ForegroundHover)
+			.SetHoveredForeground(FStyleColors::White)
+			.SetPressedForeground(FStyleColors::White)
+			.SetDisabledForeground(FStyleColors::ForegroundHover)
 			.SetNormalPadding(FMargin(5,2,5,2))
 			.SetPressedPadding(FMargin(5,3,5,1));
 		Set("GameplayTags.ChipButton.Unselected", ChipButtonUnselected);
diff --git a/Engine/Plugins/Editor/GameplayTagsEditor/Source/GameplayTagsEditor/Private/GameplayTagsSettingsCustomization.cpp b/Engine/Plugins/Editor/GameplayTagsEditor/Source/GameplayTagsEditor/Private/GameplayTagsSettingsCustomization.cpp
--- a/Engine/Plugins/Editor/GameplayTagsEditor/Source/GameplayTagsEditor/Private/GameplayTagsSettingsCustomization.cpp
+++ b/Engine/Plugins/Editor/GameplayTagsEditor/Source/GameplayTagsEditor/Private/GameplayTagsSettingsCustomization.cpp
@@ -95,17 +95,18 @@ void FGameplayTagsSettingsCustomization::CustomizeDetails(IDetailLayoutBuilder&
 					.HAlign(HAlign_Center)
 					.OnClicked_Lambda([this]()
 					{
-						const FVector2D WindowSize(320, 110);
 						const TSharedRef<SWindow> Window = SNew(SWindow)
 							.Title(LOCTEXT("AddNewGameplayTagSourceTitle", "Add new Gameplay Tag Source"))
-							.ClientSize(WindowSize)
 							.SizingRule(ESizingRule::Autosized)
-							.AutoCenter(EAutoCenter::PreferredWorkArea)
 							.SupportsMaximize(false)
 							.SupportsMinimize(false)
 							.Content()
 							[
-								SNew(SAddNewGameplayTagSourceWidget)
+								SNew(SBox)
+								.MinDesiredWidth(320.0f)
+								[
+									SNew(SAddNewGameplayTagSourceWidget)
+								]
 							];
 
 						GEditor->EditorAddModalWindow(Window);
diff --git a/Engine/Plugins/Editor/GameplayTagsEditor/Source/GameplayTagsEditor/Private/SAddNewGameplayTagSourceWidget.cpp b/Engine/Plugins/Editor/GameplayTagsEditor/Source/GameplayTagsEditor/Private/SAddNewGameplayTagSourceWidget.cpp
--- a/Engine/Plugins/Editor/GameplayTagsEditor/Source/GameplayTagsEditor/Private/SAddNewGameplayTagSourceWidget.cpp
+++ b/Engine/Plugins/Editor/GameplayTagsEditor/Source/GameplayTagsEditor/Private/SAddNewGameplayTagSourceWidget.cpp
@@ -29,18 +29,20 @@ void SAddNewGameplayTagSourceWidget::Construct(const FArguments& InArgs)
 
 	ChildSlot
 	[
-		SNew(SBorder)
-		.BorderImage(FAppStyle::GetBrush("ToolPanel.GroupBorder"))
+		SNew(SBox)
+		.Padding(InArgs._Padding)
 		[
 			SNew(SGridPanel)
+			.FillColumn(1, 1.0)
+
 			// Tag Name
 			+ SGridPanel::Slot(0, 0)
-			.Padding(5)
+			.Padding(2)
 			.VAlign(VAlign_Center)
 			.HAlign(HAlign_Left)
 			[
 				SNew(SBox)
-				.MinDesiredWidth(150.0f)
+				.MinDesiredWidth(50.0f)
 				[
 					SNew(STextBlock)
 					.Font(FAppStyle::GetFontStyle( TEXT("PropertyWindow.NormalFont")))
@@ -48,12 +50,12 @@ void SAddNewGameplayTagSourceWidget::Construct(const FArguments& InArgs)
 				]
 			]
 			+ SGridPanel::Slot(1, 0)
-			.Padding(5)
+			.Padding(2)
 			.VAlign(VAlign_Center)
-			.HAlign(HAlign_Left)
+			.HAlign(HAlign_Fill)
 			[
 				SNew(SBox)
-				.MinDesiredWidth(300.0f)
+				.MinDesiredWidth(250.0f)
 				[
 					SAssignNew(SourceNameTextBox, SEditableTextBox)
 					.HintText(HintText)
@@ -62,7 +64,7 @@ void SAddNewGameplayTagSourceWidget::Construct(const FArguments& InArgs)
 
 			// Tag source root
 			+ SGridPanel::Slot(0, 1)
-			.Padding(5)
+			.Padding(2)
 			.VAlign(VAlign_Center)
 			.HAlign(HAlign_Left)
 			[
@@ -72,28 +74,26 @@ void SAddNewGameplayTagSourceWidget::Construct(const FArguments& InArgs)
 				.ToolTipText(LOCTEXT("RootPathTooltip", "Set the base config path for added source, this includes paths from plugins and other places that call AddTagIniSearchPath"))
 			]
 			+ SGridPanel::Slot(1, 1)
-			.Padding(5)
-			.HAlign(HAlign_Left)
+			.Padding(2)
+			.VAlign(VAlign_Center)
+			.HAlign(HAlign_Fill)
 			[
-				SNew(SBox)
-				.MinDesiredWidth(300.0f)
+				SAssignNew(TagRootsComboBox, SComboBox<TSharedPtr<FString> >)
+				.OptionsSource(&TagRoots)
+				.OnGenerateWidget(this, &SAddNewGameplayTagSourceWidget::OnGenerateTagRootsComboBox)
+				.Content()
 				[
-					SAssignNew(TagRootsComboBox, SComboBox<TSharedPtr<FString> >)
-					.OptionsSource(&TagRoots)
-					.OnGenerateWidget(this, &SAddNewGameplayTagSourceWidget::OnGenerateTagRootsComboBox)
-					.Content()
-					[
-						SNew(STextBlock)
-						.Text(this, &SAddNewGameplayTagSourceWidget::CreateTagRootsComboBoxContent)
-						.Font(IDetailLayoutBuilder::GetDetailFont())
-					]
+					SNew(STextBlock)
+					.Text(this, &SAddNewGameplayTagSourceWidget::CreateTagRootsComboBoxContent)
+					.Font(IDetailLayoutBuilder::GetDetailFont())
 				]
 			]
 			
 			// Add Source Button
-			+ SGridPanel::Slot(1, 2)
-			.Padding(5)
-			.HAlign(HAlign_Left)
+			+ SGridPanel::Slot(0, 2)
+			.ColumnSpan(2)
+			.Padding(FMargin(0, 16))
+			.HAlign(HAlign_Right)
 			[
 				SNew(SButton)
 				.Text(LOCTEXT("AddNew", "Add New Source"))
diff --git a/Engine/Plugins/Editor/GameplayTagsEditor/Source/GameplayTagsEditor/Private/SAddNewGameplayTagSourceWidget.h b/Engine/Plugins/Editor/GameplayTagsEditor/Source/GameplayTagsEditor/Private/SAddNewGameplayTagSourceWidget.h
--- a/Engine/Plugins/Editor/GameplayTagsEditor/Source/GameplayTagsEditor/Private/SAddNewGameplayTagSourceWidget.h
+++ b/Engine/Plugins/Editor/GameplayTagsEditor/Source/GameplayTagsEditor/Private/SAddNewGameplayTagSourceWidget.h
@@ -16,9 +16,11 @@ public:
 
 	SLATE_BEGIN_ARGS(SAddNewGameplayTagSourceWidget)
 		: _NewSourceName(TEXT(""))
+		, _Padding(FMargin(15))
 		{}
 		SLATE_EVENT(FOnGameplayTagSourceAdded, OnGameplayTagSourceAdded )	// Callback for when a new source is added	
 		SLATE_ARGUMENT( FString, NewSourceName ) // String that will initially populate the New Source Name field
+		SLATE_ARGUMENT(FMargin, Padding)
 	SLATE_END_ARGS();
 
 	virtual void Tick( const FGeometry& AllottedGeometry, const double InCurrentTime, const float InDeltaTime ) override;
diff --git a/Engine/Plugins/Editor/GameplayTagsEditor/Source/GameplayTagsEditor/Private/SAddNewGameplayTagWidget.cpp b/Engine/Plugins/Editor/GameplayTagsEditor/Source/GameplayTagsEditor/Private/SAddNewGameplayTagWidget.cpp
--- a/Engine/Plugins/Editor/GameplayTagsEditor/Source/GameplayTagsEditor/Private/SAddNewGameplayTagWidget.cpp
+++ b/Engine/Plugins/Editor/GameplayTagsEditor/Source/GameplayTagsEditor/Private/SAddNewGameplayTagWidget.cpp
@@ -44,42 +44,36 @@ void SAddNewGameplayTagWidget::Construct(const FArguments& InArgs)
 
 	ChildSlot
 	[
-		SNew(SBorder)
-		.BorderImage(FAppStyle::GetBrush("ToolPanel.GroupBorder"))
+		SNew(SBox)
+		.Padding(InArgs._Padding)
 		[
 			SNew(SGridPanel)
+			.FillColumn(1, 1.0)
 			
 			// Tag Name
 			+ SGridPanel::Slot(0, 0)
-			.Padding(5)
+			.Padding(2)
 			.VAlign(VAlign_Center)
 			.HAlign(HAlign_Left)
 			[
-				SNew(SBox)
-				.MinDesiredWidth(150.0f)
-				[
-					SNew(STextBlock)
-					.Font(FAppStyle::GetFontStyle( TEXT("PropertyWindow.NormalFont")))
-					.Text(LOCTEXT("NewTagName", "Name:"))
-				]
+				SNew(STextBlock)
+				.Font(FAppStyle::GetFontStyle( TEXT("PropertyWindow.NormalFont")))
+				.Text(LOCTEXT("NewTagName", "Name:"))
 			]
 			+ SGridPanel::Slot(1, 0)
-			.Padding(5)
+			.Padding(2)
 			.VAlign(VAlign_Center)
-			.HAlign(HAlign_Left)
+			.HAlign(HAlign_Fill)
 			[
-				SNew(SBox)
-				.MinDesiredWidth(300.0f)
-				[
-					SAssignNew(TagNameTextBox, SEditableTextBox)
-					.HintText(HintText)
-					.OnTextCommitted(this, &SAddNewGameplayTagWidget::OnCommitNewTagName)
-				]
+				SAssignNew(TagNameTextBox, SEditableTextBox)
+				.HintText(HintText)
+				.Font(FAppStyle::GetFontStyle( TEXT("PropertyWindow.NormalFont")))
+				.OnTextCommitted(this, &SAddNewGameplayTagWidget::OnCommitNewTagName)
 			]
 			
 			// Tag Comment
 			+ SGridPanel::Slot(0, 1)
-			.Padding(5)
+			.Padding(2)
 			.VAlign(VAlign_Center)
 			.HAlign(HAlign_Left)
 			[
@@ -88,78 +82,72 @@ void SAddNewGameplayTagWidget::Construct(const FArguments& InArgs)
 				.Text(LOCTEXT("TagComment", "Comment:"))
 			]
 			+ SGridPanel::Slot(1, 1)
-			.Padding(5)
+			.Padding(2)
 			.VAlign(VAlign_Center)
-			.HAlign(HAlign_Left)
+			.HAlign(HAlign_Fill)
 			[
-				SNew(SBox)
-				.MinDesiredWidth(300.0f)
-				[
-					SAssignNew(TagCommentTextBox, SEditableTextBox)
-					.HintText(LOCTEXT("TagCommentHint", "Comment"))
-					.OnTextCommitted(this, &SAddNewGameplayTagWidget::OnCommitNewTagName)
-				]
+				SAssignNew(TagCommentTextBox, SEditableTextBox)
+				.HintText(LOCTEXT("TagCommentHint", "Comment"))
+				.Font(FAppStyle::GetFontStyle( TEXT("PropertyWindow.NormalFont")))
+				.OnTextCommitted(this, &SAddNewGameplayTagWidget::OnCommitNewTagName)
 			]
 
 			// Tag Location
 			+ SGridPanel::Slot(0, 2)
-			.Padding(5)
+			.Padding(2)
 			.VAlign(VAlign_Center)
 			.HAlign(HAlign_Left)
 			[
 				SNew(STextBlock)
 				.Text(LOCTEXT("CreateTagSource", "Source:"))
-				.Font(IDetailLayoutBuilder::GetDetailFont())
+				.Font(FAppStyle::GetFontStyle( TEXT("PropertyWindow.NormalFont")))
 			]
 			+ SGridPanel::Slot(1, 2)
-			.Padding(5)
+			.Padding(2)
 			.VAlign(VAlign_Center)
-			.HAlign(HAlign_Left)
+			.HAlign(HAlign_Fill)
 			[
-				SNew(SBox)
-				.MinDesiredWidth(300.0f)
-				[
-					SNew(SHorizontalBox)
+				SNew(SHorizontalBox)
 
-					+ SHorizontalBox::Slot()
-					.Padding(2.0f, 2.0f)
-					.FillWidth(1.0f)
-					.VAlign(VAlign_Center)
+				+ SHorizontalBox::Slot()
+				.FillWidth(1.0f)
+				.VAlign(VAlign_Center)
+				[
+					SAssignNew(TagSourcesComboBox, SComboBox<TSharedPtr<FName> >)
+					.OptionsSource(&TagSources)
+					.OnGenerateWidget(this, &SAddNewGameplayTagWidget::OnGenerateTagSourcesComboBox)
+					.ToolTipText(this, &SAddNewGameplayTagWidget::CreateTagSourcesComboBoxToolTip)
+					.Content()
 					[
-						SAssignNew(TagSourcesComboBox, SComboBox<TSharedPtr<FName> >)
-						.OptionsSource(&TagSources)
-						.OnGenerateWidget(this, &SAddNewGameplayTagWidget::OnGenerateTagSourcesComboBox)
-						.ToolTipText(this, &SAddNewGameplayTagWidget::CreateTagSourcesComboBoxToolTip)
-						.Content()
-						[
-							SNew(STextBlock)
-							.Text(this, &SAddNewGameplayTagWidget::CreateTagSourcesComboBoxContent)
-							.Font(IDetailLayoutBuilder::GetDetailFont())
-						]
+						SNew(STextBlock)
+						.Text(this, &SAddNewGameplayTagWidget::CreateTagSourcesComboBoxContent)
+						.Font(IDetailLayoutBuilder::GetDetailFont())
 					]
+				]
 
-					+ SHorizontalBox::Slot()
-					.AutoWidth()
-					.VAlign(VAlign_Center)
+				+ SHorizontalBox::Slot()
+				.AutoWidth()
+				.Padding(2.0f, 0, 0, 0)
+				.VAlign(VAlign_Center)
+				[
+					SNew( SButton )
+					.ButtonStyle( FAppStyle::Get(), "NoBorder" )
+					.Visibility(this, &SAddNewGameplayTagWidget::OnGetTagSourceFavoritesVisibility)
+					.OnClicked(this, &SAddNewGameplayTagWidget::OnToggleTagSourceFavoriteClicked)
+					.ToolTipText(LOCTEXT("ToggleFavoriteTooltip", "Toggle whether or not this tag source is your favorite source (new tags will go into your favorite source by default)"))
+					.ContentPadding(0)
 					[
-						SNew( SButton )
-						.ButtonStyle( FAppStyle::Get(), "NoBorder" )
-						.Visibility(this, &SAddNewGameplayTagWidget::OnGetTagSourceFavoritesVisibility)
-						.OnClicked(this, &SAddNewGameplayTagWidget::OnToggleTagSourceFavoriteClicked)
-						.ToolTipText(LOCTEXT("ToggleFavoriteTooltip", "Toggle whether or not this tag source is your favorite source (new tags will go into your favorite source by default)"))
-						.ContentPadding(0)
-						[
-							SNew(SImage)
-							.Image(this, &SAddNewGameplayTagWidget::OnGetTagSourceFavoriteImage)
-						]
+						SNew(SImage)
+						.Image(this, &SAddNewGameplayTagWidget::OnGetTagSourceFavoriteImage)
 					]
 				]
 			]
 
 			// Add Tag Button
-			+ SGridPanel::Slot(1, 3)
-			.Padding(5)
-			.HAlign(HAlign_Left)
+			+ SGridPanel::Slot(0, 3)
+			.ColumnSpan(2)
+			.Padding(InArgs._AddButtonPadding)
+			.HAlign(HAlign_Right)
 			[
 				SNew(SButton)
 				.Text(LOCTEXT("AddNew", "Add New Tag"))
@@ -202,6 +190,7 @@ void SAddNewGameplayTagWidget::Tick( const FGeometry& AllottedGeometry, const do
 	{
 		bShouldGetKeyboardFocus = false;
 		FSlateApplication::Get().SetKeyboardFocus(TagNameTextBox.ToSharedRef(), EFocusCause::SetDirectly);
+		FSlateApplication::Get().SetUserFocus(0, TagNameTextBox.ToSharedRef());
 	}
 }
 
diff --git a/Engine/Plugins/Editor/GameplayTagsEditor/Source/GameplayTagsEditor/Private/SAddNewRestrictedGameplayTagWidget.cpp b/Engine/Plugins/Editor/GameplayTagsEditor/Source/GameplayTagsEditor/Private/SAddNewRestrictedGameplayTagWidget.cpp
--- a/Engine/Plugins/Editor/GameplayTagsEditor/Source/GameplayTagsEditor/Private/SAddNewRestrictedGameplayTagWidget.cpp
+++ b/Engine/Plugins/Editor/GameplayTagsEditor/Source/GameplayTagsEditor/Private/SAddNewRestrictedGameplayTagWidget.cpp
@@ -43,42 +43,35 @@ void SAddNewRestrictedGameplayTagWidget::Construct(const FArguments& InArgs)
 
 	ChildSlot
 	[
-		SNew(SBorder)
-		.BorderImage(FAppStyle::GetBrush("ToolPanel.GroupBorder"))
+		SNew(SBox)
+		.Padding(InArgs._Padding)
 		[
 			SNew(SGridPanel)
-			
+			.FillColumn(1, 1.0)
+
 			// Restricted Tag Name
 			+ SGridPanel::Slot(0, 0)
-			.Padding(5)
+			.Padding(2)
 			.VAlign(VAlign_Center)
 			.HAlign(HAlign_Left)
 			[
-				SNew(SBox)
-				.MinDesiredWidth(150.0f)
-				[
-					SNew(STextBlock)
-					.Font(FAppStyle::GetFontStyle( TEXT("PropertyWindow.NormalFont")))
-					.Text(LOCTEXT("NewTagName", "Name:"))
-				]
+				SNew(STextBlock)
+				.Font(FAppStyle::GetFontStyle( TEXT("PropertyWindow.NormalFont")))
+				.Text(LOCTEXT("NewTagName", "Name:"))
 			]
 			+ SGridPanel::Slot(1, 0)
-			.Padding(5)
+			.Padding(2)
 			.VAlign(VAlign_Center)
-			.HAlign(HAlign_Left)
+			.HAlign(HAlign_Fill)
 			[
-				SNew(SBox)
-				.MinDesiredWidth(300.0f)
-				[
-					SAssignNew(TagNameTextBox, SEditableTextBox)
-					.HintText(HintText)
-					.OnTextCommitted(this, &SAddNewRestrictedGameplayTagWidget::OnCommitNewTagName)
-				]
+				SAssignNew(TagNameTextBox, SEditableTextBox)
+				.HintText(HintText)
+				.OnTextCommitted(this, &SAddNewRestrictedGameplayTagWidget::OnCommitNewTagName)
 			]
 			
 			// Tag Comment
 			+ SGridPanel::Slot(0, 1)
-			.Padding(5)
+			.Padding(2)
 			.VAlign(VAlign_Center)
 			.HAlign(HAlign_Left)
 			[
@@ -87,36 +80,35 @@ void SAddNewRestrictedGameplayTagWidget::Construct(const FArguments& InArgs)
 				.Text(LOCTEXT("TagComment", "Comment:"))
 			]
 			+ SGridPanel::Slot(1, 1)
-			.Padding(5)
+			.Padding(2)
 			.VAlign(VAlign_Center)
-			.HAlign(HAlign_Left)
+			.HAlign(HAlign_Fill)
 			[
-				SNew(SBox)
-				.MinDesiredWidth(300.0f)
-				[
-					SAssignNew(TagCommentTextBox, SEditableTextBox)
-					.HintText(LOCTEXT("TagCommentHint", "Comment"))
-					.OnTextCommitted(this, &SAddNewRestrictedGameplayTagWidget::OnCommitNewTagName)
-				]
+				SAssignNew(TagCommentTextBox, SEditableTextBox)
+				.HintText(LOCTEXT("TagCommentHint", "Comment"))
+				.OnTextCommitted(this, &SAddNewRestrictedGameplayTagWidget::OnCommitNewTagName)
 			]
 
 			// Allow non-restricted children
 			+ SGridPanel::Slot(0, 2)
-			.Padding(5)
+			.Padding(2)
+			.ColumnSpan(2)
 			.VAlign(VAlign_Center)
 			.HAlign(HAlign_Left)
 			[
-				SNew(STextBlock)
-				.Font(FAppStyle::GetFontStyle( TEXT("PropertyWindow.NormalFont")))
-				.Text(LOCTEXT("AllowNonRestrictedChildren", "Allow non-restricted children:"))
-			]
-			+ SGridPanel::Slot(1, 2)
-			.Padding(5)
-			.VAlign(VAlign_Center)
-			.HAlign(HAlign_Left)
-			[
-				SNew(SBox)
-				.MinDesiredWidth(300.0f)
+				SNew(SHorizontalBox)
+				+ SHorizontalBox::Slot()
+				.VAlign(VAlign_Center)
+				.AutoWidth()
+				[
+					SNew(STextBlock)
+					.Font(FAppStyle::GetFontStyle( TEXT("PropertyWindow.NormalFont")))
+					.Text(LOCTEXT("AllowNonRestrictedChildren", "Allow non-restricted children:"))
+				]
+				+ SHorizontalBox::Slot()
+				.Padding(FMargin(4,0))
+				.VAlign(VAlign_Center)
+				.AutoWidth()
 				[
 					SAssignNew(AllowNonRestrictedChildrenCheckBox, SCheckBox)
 				]
@@ -124,7 +116,7 @@ void SAddNewRestrictedGameplayTagWidget::Construct(const FArguments& InArgs)
 
 			// Tag Location
 			+ SGridPanel::Slot(0, 3)
-			.Padding(5)
+			.Padding(2)
 			.VAlign(VAlign_Center)
 			.HAlign(HAlign_Left)
 			[
@@ -133,29 +125,26 @@ void SAddNewRestrictedGameplayTagWidget::Construct(const FArguments& InArgs)
 				.Font(IDetailLayoutBuilder::GetDetailFont())
 			]
 			+ SGridPanel::Slot(1, 3)
-			.Padding(5)
+			.Padding(2)
 			.VAlign(VAlign_Center)
-			.HAlign(HAlign_Left)
+			.HAlign(HAlign_Fill)
 			[
-				SNew(SBox)
-				.MinDesiredWidth(300.0f)
+				SAssignNew(TagSourcesComboBox, SComboBox<TSharedPtr<FName> >)
+				.OptionsSource(&RestrictedTagSources)
+				.OnGenerateWidget(this, &SAddNewRestrictedGameplayTagWidget::OnGenerateTagSourcesComboBox)
+				.Content()
 				[
-					SAssignNew(TagSourcesComboBox, SComboBox<TSharedPtr<FName> >)
-					.OptionsSource(&RestrictedTagSources)
-					.OnGenerateWidget(this, &SAddNewRestrictedGameplayTagWidget::OnGenerateTagSourcesComboBox)
-					.Content()
-					[
-						SNew(STextBlock)
-						.Text(this, &SAddNewRestrictedGameplayTagWidget::CreateTagSourcesComboBoxContent)
-						.Font(IDetailLayoutBuilder::GetDetailFont())
-					]
+					SNew(STextBlock)
+					.Text(this, &SAddNewRestrictedGameplayTagWidget::CreateTagSourcesComboBoxContent)
+					.Font(IDetailLayoutBuilder::GetDetailFont())
 				]
 			]
 
 			// Add Tag Button
-			+ SGridPanel::Slot(1, 4)
-			.Padding(5)
-			.HAlign(HAlign_Left)
+			+ SGridPanel::Slot(0, 4)
+			.ColumnSpan(2)
+			.Padding(FMargin(0, 16))
+			.HAlign(HAlign_Right)
 			[
 				SNew(SButton)
 				.Text(LOCTEXT("AddNew", "Add New Tag"))
diff --git a/Engine/Plugins/Editor/GameplayTagsEditor/Source/GameplayTagsEditor/Private/SAddNewRestrictedGameplayTagWidget.h b/Engine/Plugins/Editor/GameplayTagsEditor/Source/GameplayTagsEditor/Private/SAddNewRestrictedGameplayTagWidget.h
--- a/Engine/Plugins/Editor/GameplayTagsEditor/Source/GameplayTagsEditor/Private/SAddNewRestrictedGameplayTagWidget.h
+++ b/Engine/Plugins/Editor/GameplayTagsEditor/Source/GameplayTagsEditor/Private/SAddNewRestrictedGameplayTagWidget.h
@@ -19,9 +19,11 @@ public:
 
 	SLATE_BEGIN_ARGS(SAddNewRestrictedGameplayTagWidget)
 		: _NewRestrictedTagName(TEXT(""))
+		, _Padding(FMargin(15))
 		{}
 		SLATE_EVENT( FOnRestrictedGameplayTagAdded, OnRestrictedGameplayTagAdded )	// Callback for when a new tag is added	
 		SLATE_ARGUMENT( FString, NewRestrictedTagName ) // String that will initially populate the New Tag Name field
+		SLATE_ARGUMENT(FMargin, Padding)
 	SLATE_END_ARGS();
 
 	virtual ~SAddNewRestrictedGameplayTagWidget() override;
diff --git a/Engine/Plugins/Editor/GameplayTagsEditor/Source/GameplayTagsEditor/Private/SGameplayTagPicker.cpp b/Engine/Plugins/Editor/GameplayTagsEditor/Source/GameplayTagsEditor/Private/SGameplayTagPicker.cpp
--- a/Engine/Plugins/Editor/GameplayTagsEditor/Source/GameplayTagsEditor/Private/SGameplayTagPicker.cpp
+++ b/Engine/Plugins/Editor/GameplayTagsEditor/Source/GameplayTagsEditor/Private/SGameplayTagPicker.cpp
@@ -4,6 +4,7 @@
 #include "Misc/ConfigCacheIni.h"
 #include "Widgets/Input/SCheckBox.h"
 #include "Widgets/Input/SComboButton.h"
+#include "Widgets/Input/SButton.h"
 #include "Widgets/Images/SImage.h"
 #include "Styling/AppStyle.h"
 #include "Widgets/SWindow.h"
@@ -147,6 +148,9 @@ void SGameplayTagPicker::Construct(const FArguments& InArgs)
 		];
 	SettingsCombo->SetOnGetMenuContent(FOnGetContent::CreateSP(this, &SGameplayTagPicker::MakeSettingsMenu, SettingsCombo));
 
+
+	TWeakPtr<SGameplayTagPicker> WeakSelf = StaticCastWeakPtr<SGameplayTagPicker>(AsWeak());
+	
 	TSharedRef<SWidget> Picker = 
 		SNew(SBorder)
 		.Padding(InArgs._Padding)
@@ -161,7 +165,7 @@ void SGameplayTagPicker::Construct(const FArguments& InArgs)
 			[
 				SNew(SHorizontalBox)
 
-				// Add button
+				// Add button for management mode
 				+SHorizontalBox::Slot()
 				.AutoWidth()
 				[
@@ -180,6 +184,42 @@ void SGameplayTagPicker::Construct(const FArguments& InArgs)
 					]
 				]
 
+				// Smaller add button for selection and hybrid modes.
+				+SHorizontalBox::Slot()
+				.AutoWidth()
+				[
+					SNew(SHorizontalBox)
+
+					+SHorizontalBox::Slot()
+					.AutoWidth()
+					.Padding(FMargin(0,0,4,0))
+					[
+						SNew(SButton)
+						.ButtonStyle(FAppStyle::Get(), "HoverHintOnly")
+						.Visibility(InArgs._GameplayTagPickerMode != EGameplayTagPickerMode::ManagementMode ? EVisibility::Visible : EVisibility::Collapsed)
+						.OnClicked_Lambda([WeakSelf]()
+						{
+							if (const TSharedPtr<SGameplayTagPicker> Self = WeakSelf.Pin())
+							{
+								if (!Self->bNewTagWidgetVisible)
+								{
+									Self->ShowInlineAddTagWidget(EGameplayTagAdd::Root);
+								}
+								else
+								{
+									Self->bNewTagWidgetVisible = false;
+								}
+							}
+							return FReply::Handled();
+						})
+						[
+							SNew(SImage)
+                        	.Image(FAppStyle::GetBrush("Icons.Plus"))
+                        	.ColorAndOpacity(FStyleColors::AccentGreen)
+						]
+					]
+				]
+
 				// Search
 				+SHorizontalBox::Slot()
 				.VAlign(VAlign_Center)
@@ -199,6 +239,30 @@ void SGameplayTagPicker::Construct(const FArguments& InArgs)
 				]
 			]
 
+			// Inline add new tag window for selection and hybrid modes.
+			+SVerticalBox::Slot()
+			.AutoHeight()
+			.VAlign(VAlign_Top)
+			.Padding(FMargin(4,2))
+			[
+				SAssignNew(AddNewTagWidget, SAddNewGameplayTagWidget )
+				.Padding(FMargin(0, 2))
+				.AddButtonPadding(FMargin(0, 4, 0, 0))
+				.Visibility_Lambda([WeakSelf]()
+				{
+					const TSharedPtr<SGameplayTagPicker> Self = WeakSelf.Pin();
+					return (Self.IsValid() && Self->bNewTagWidgetVisible) ? EVisibility::Visible : EVisibility::Collapsed;
+					
+				})
+				.OnGameplayTagAdded_Lambda([WeakSelf](const FString& TagName, const FString& TagComment, const FName& TagSource)
+				{
+				   if (const TSharedPtr<SGameplayTagPicker> Self = WeakSelf.Pin())
+				   {
+					   Self->OnGameplayTagAdded(TagName, TagComment, TagSource);
+				   }
+				})
+			]
+
 			// Gameplay Tags tree
 			+SVerticalBox::Slot()
 			.MaxHeight(MaxHeight)
@@ -240,7 +304,7 @@ void SGameplayTagPicker::Construct(const FArguments& InArgs)
 
 		MenuBuilder.AddMenuEntry(
 			LOCTEXT("GameplayTagPicker_ManageTags", "Manage Gameplay Tags..."), FText::GetEmpty(), FSlateIcon(FAppStyle::GetAppStyleSetName(), "Icons.Settings"),
-			FUIAction(FExecuteAction::CreateRaw(this, &SGameplayTagPicker::OnManageTagsClicked, TSharedPtr<SComboButton>()))
+			FUIAction(FExecuteAction::CreateRaw(this, &SGameplayTagPicker::OnManageTagsClicked, TSharedPtr<FGameplayTagNode>(), TSharedPtr<SComboButton>()))
 		);
 
 		MenuBuilder.AddSeparator();
@@ -354,7 +418,7 @@ TSharedRef<SWidget> SGameplayTagPicker::MakeAddMenu()
 	// Add root tag
 	MenuBuilder.AddMenuEntry(
 		LOCTEXT("GameplayTagPicker_AddNewRootTag", "Add new root level tag"), FText::GetEmpty(), FSlateIcon(FAppStyle::GetAppStyleSetName(), "Icons.Plus"),
-		FUIAction(FExecuteAction::CreateRaw(this, &SGameplayTagPicker::AddTagDialog, EGameplayTagAdd::Root, TSharedPtr<FGameplayTagNode>()))
+		FUIAction(FExecuteAction::CreateRaw(this, &SGameplayTagPicker::OpenAddTagDialog, EGameplayTagAdd::Root, TSharedPtr<FGameplayTagNode>()))
 	);
 
 	// Add child tag
@@ -372,25 +436,6 @@ TSharedRef<SWidget> SGameplayTagPicker::MakeAddMenu()
 	return MenuBuilder.MakeWidget();
 }
 
-TSharedRef<SWidget> SGameplayTagPicker::MakeActionsMenu(TSharedPtr<SComboButton> OwnerCombo)
-{
-	FMenuBuilder MenuBuilder(/*bInShouldCloseWindowAfterMenuSelection*/false, nullptr);
-
-	MenuBuilder.AddMenuEntry(
-		LOCTEXT("GameplayTagPicker_CreateAndManageTags", "Create and Manage Gameplay Tags..."), FText::GetEmpty(), FSlateIcon(FAppStyle::GetAppStyleSetName(), "Icons.Settings"),
-		FUIAction(FExecuteAction::CreateRaw(this, &SGameplayTagPicker::OnManageTagsClicked, OwnerCombo))
-	);
-	
-	MenuBuilder.AddMenuEntry(
-		LOCTEXT("GameplayTagPicker_ClearSelection", "Clear Selection"), FText::GetEmpty(), FSlateIcon(FAppStyle::GetAppStyleSetName(), "Icons.X"),
-		FUIAction(FExecuteAction::CreateRaw(this, &SGameplayTagPicker::OnClearAllClicked, OwnerCombo))
-	);
-
-	MenuBuilder.AddSeparator();
-	
-	return MenuBuilder.MakeWidget();
-}
-
 TSharedRef<SWidget> SGameplayTagPicker::MakeSettingsMenu(TSharedPtr<SComboButton> OwnerCombo)
 {
 	FMenuBuilder MenuBuilder(/*bInShouldCloseWindowAfterMenuSelection*/false, nullptr);
@@ -430,7 +475,21 @@ void SGameplayTagPicker::Tick(const FGeometry& AllottedGeometry, const double In
 		// Scroll specified item into view.
 		UGameplayTagsManager& Manager = UGameplayTagsManager::Get();
 		TSharedPtr<FGameplayTagNode> Node = Manager.FindTagNode(RequestedScrollToTag);
-		TagTreeWidget->RequestScrollIntoView(Node);
+
+		if (Node.IsValid())
+		{
+			// Expand all the parent nodes to make sure the target node is visible.
+			TSharedPtr<FGameplayTagNode> ParentNode = Node.IsValid() ? Node->ParentNode : nullptr;
+			while (ParentNode.IsValid())
+			{
+				TagTreeWidget->SetItemExpansion(ParentNode, /*bExpand*/true);
+				ParentNode = ParentNode->ParentNode;
+			}
+
+			TagTreeWidget->SetItemSelection(Node, true);
+			TagTreeWidget->RequestScrollIntoView(Node);
+		}
+		
 		RequestedScrollToTag = FGameplayTag();
 	}
 }
@@ -942,11 +1001,16 @@ EVisibility SGameplayTagPicker::DetermineAllowChildrenVisible(TSharedPtr<FGamepl
 	return EVisibility::Collapsed;
 }
 
-void SGameplayTagPicker::OnManageTagsClicked(TSharedPtr<SComboButton> OwnerCombo)
+void SGameplayTagPicker::OnManageTagsClicked(TSharedPtr<FGameplayTagNode> Node, TSharedPtr<SComboButton> OwnerCombo)
 {
 	FGameplayTagManagerWindowArgs Args;
 	Args.bRestrictedTags = bRestrictedTags;
 	Args.Filter = RootFilterString;
+	if (Node.IsValid())
+	{
+		Args.HighlightedTag = Node->GetCompleteTag();
+	}
+	
 	UE::GameplayTags::Editor::OpenGameplayTagManager(Args);
 
 	if (OwnerCombo.IsValid())
@@ -1015,26 +1079,11 @@ void SGameplayTagPicker::OnCollapseAllClicked(TSharedPtr<SComboButton> OwnerComb
 	}
 }
 
-FReply SGameplayTagPicker::OnAddSubtagClicked(TSharedPtr<FGameplayTagNode> InTagNode)
+void SGameplayTagPicker::OpenAddTagDialog(const EGameplayTagAdd Mode, TSharedPtr<FGameplayTagNode> InTagNode)
 {
-	if (!bReadOnly)
-	{
-		AddTagDialog(EGameplayTagAdd::Child, InTagNode);
-	}
-	return FReply::Handled();
-}
-
-void SGameplayTagPicker::AddTagDialog(const EGameplayTagAdd Mode, TSharedPtr<FGameplayTagNode> InTagNode)
-{
-	const FVector2D WindowSize(320, 110);
-	const FVector2D WindowPosition = FSlateApplication::Get().GetCursorPos() - WindowSize * 0.5f;
-	
 	TSharedPtr<SWindow> NewTagWindow = SNew(SWindow)
 		.Title(LOCTEXT("EditTagWindowTitle", "Edit Gameplay Tag"))
-		.ClientSize(WindowSize)
-		.ScreenPosition(WindowPosition)
 		.SizingRule(ESizingRule::Autosized)
-		.AutoCenter(EAutoCenter::None)
 		.SupportsMaximize(false)
 		.SupportsMinimize(false);
 
@@ -1060,7 +1109,7 @@ void SGameplayTagPicker::AddTagDialog(const EGameplayTagAdd Mode, TSharedPtr<FGa
 	
 	if (!bRestrictedTags)
     {
-		TSharedRef<SAddNewGameplayTagWidget> AddNewTagDialog = SNew(SAddNewGameplayTagWidget )
+		TSharedRef<SAddNewGameplayTagWidget> AddNewTagDialog = SNew(SAddNewGameplayTagWidget)
 			.OnGameplayTagAdded_Lambda([WeakSelf, NewTagWindow](const FString& TagName, const FString& TagComment, const FName& TagSource)
 			{
 				if (const TSharedPtr<SGameplayTagPicker> Self = WeakSelf.Pin())
@@ -1082,8 +1131,11 @@ void SGameplayTagPicker::AddTagDialog(const EGameplayTagAdd Mode, TSharedPtr<FGa
 			AddNewTagDialog->AddDuplicate(TagName, TagSource);
 		}
 
-		NewTagWindow->SetContent(AddNewTagDialog);
-
+		NewTagWindow->SetContent(SNew(SBox)
+			.MinDesiredWidth(320.0f)
+			[
+				AddNewTagDialog
+			]);
     }
     else if (bRestrictedTags)
     {
@@ -1109,16 +1161,54 @@ void SGameplayTagPicker::AddTagDialog(const EGameplayTagAdd Mode, TSharedPtr<FGa
     		AddNewTagDialog->AddDuplicate(TagName, TagSource, bTagAllowsNonRestrictedChildren);
     	}
         
-		NewTagWindow->SetContent(AddNewTagDialog);
+    	NewTagWindow->SetContent(SNew(SBox)
+			.MinDesiredWidth(320.0f)
+			[
+				AddNewTagDialog
+			]);
     }
 
 	TSharedPtr<SWindow> CurrentWindow = FSlateApplication::Get().FindWidgetWindow(AsShared());
 	FSlateApplication::Get().AddModalWindow(NewTagWindow.ToSharedRef(), CurrentWindow);
 }
 
+void SGameplayTagPicker::ShowInlineAddTagWidget(const EGameplayTagAdd Mode, TSharedPtr<FGameplayTagNode> InTagNode)
+{
+	UGameplayTagsManager& Manager = UGameplayTagsManager::Get();
+	
+	FString TagName;
+	FName TagFName;
+	FString TagComment;
+	FName TagSource;
+	bool bTagIsExplicit;
+	bool bTagIsRestricted;
+	bool bTagAllowsNonRestrictedChildren;
+
+	if (InTagNode.IsValid())
+	{
+		TagName = InTagNode->GetCompleteTagString();
+		TagFName = InTagNode->GetCompleteTagName();
+	}
+
+	Manager.GetTagEditorData(TagFName, TagComment, TagSource, bTagIsExplicit, bTagIsRestricted, bTagAllowsNonRestrictedChildren);
+
+	if (AddNewTagWidget.IsValid())
+	{
+		if (Mode == EGameplayTagAdd::Child || Mode == EGameplayTagAdd::Root)
+		{
+			AddNewTagWidget->AddSubtagFromParent(TagName, TagSource);
+		}
+		else if (Mode == EGameplayTagAdd::Duplicate)
+		{
+			AddNewTagWidget->AddDuplicate(TagName, TagSource);
+		}
+		bNewTagWidgetVisible = true;
+	}
+}
+
 FReply SGameplayTagPicker::OnAddRootTagClicked()
 {
-	AddTagDialog(EGameplayTagAdd::Root);
+	OpenAddTagDialog(EGameplayTagAdd::Root);
 	
 	return FReply::Handled();
 }
@@ -1148,22 +1238,22 @@ TSharedRef<SWidget> SGameplayTagPicker::MakeTagActionsMenu(TSharedPtr<FGameplayT
 	// Occurs when SGameplayTagPicker is being used as a menu item itself (Details panel of blueprint editor for example).
 	FMenuBuilder MenuBuilder(bInShouldCloseWindowAfterMenuSelection, nullptr);
 
-	if (bShowManagement)
-	{
-		// Add child tag
-		MenuBuilder.AddMenuEntry(LOCTEXT("GameplayTagPicker_AddSubTag", "Add Sub Tag"),
-			LOCTEXT("GameplayTagPicker_AddSubTagTagTooltip", "Add sub tag under selected tag."),
-			FSlateIcon(FAppStyle::GetAppStyleSetName(), "Icons.Plus"),
-			FUIAction(FExecuteAction::CreateSP(this, &SGameplayTagPicker::OnAddSubTag, InTagNode, ActionsCombo), FCanExecuteAction::CreateSP(this, &SGameplayTagPicker::CanAddNewSubTag, InTagNode)));
-	
-		// Duplicate
-		MenuBuilder.AddMenuEntry(LOCTEXT("GameplayTagPicker_DuplicateTag", "Duplicate Tag"),
-			LOCTEXT("GameplayTagPicker_DuplicateTagTooltip", "Duplicate selected tag to create a new tag."),
-			FSlateIcon(FAppStyle::GetAppStyleSetName(), "Icons.Duplicate"),
-			FUIAction(FExecuteAction::CreateSP(this, &SGameplayTagPicker::OnDuplicateTag, InTagNode, ActionsCombo), FCanExecuteAction::CreateSP(this, &SGameplayTagPicker::CanAddNewSubTag, InTagNode)));
+	// Add child tag
+	MenuBuilder.AddMenuEntry(LOCTEXT("GameplayTagPicker_AddSubTag", "Add Sub Tag"),
+		LOCTEXT("GameplayTagPicker_AddSubTagTagTooltip", "Add sub tag under selected tag."),
+		FSlateIcon(FAppStyle::GetAppStyleSetName(), "Icons.Plus"),
+		FUIAction(FExecuteAction::CreateSP(this, &SGameplayTagPicker::OnAddSubTag, InTagNode, ActionsCombo), FCanExecuteAction::CreateSP(this, &SGameplayTagPicker::CanAddNewSubTag, InTagNode)));
 
-		MenuBuilder.AddSeparator();
+	// Duplicate
+	MenuBuilder.AddMenuEntry(LOCTEXT("GameplayTagPicker_DuplicateTag", "Duplicate Tag"),
+		LOCTEXT("GameplayTagPicker_DuplicateTagTooltip", "Duplicate selected tag to create a new tag."),
+		FSlateIcon(FAppStyle::GetAppStyleSetName(), "Icons.Duplicate"),
+		FUIAction(FExecuteAction::CreateSP(this, &SGameplayTagPicker::OnDuplicateTag, InTagNode, ActionsCombo), FCanExecuteAction::CreateSP(this, &SGameplayTagPicker::CanAddNewSubTag, InTagNode)));
+
+	MenuBuilder.AddSeparator();
 
+	if (bShowManagement)
+	{
 		// Rename
 		MenuBuilder.AddMenuEntry(LOCTEXT("GameplayTagPicker_RenameTag", "Rename Tag"),
 			LOCTEXT("GameplayTagPicker_RenameTagTooltip", "Rename this tag"),
@@ -1182,7 +1272,7 @@ TSharedRef<SWidget> SGameplayTagPicker::MakeTagActionsMenu(TSharedPtr<FGameplayT
 	// Only include these menu items if we have tag containers to modify
 	if (bMultiSelect)
 	{
-		// Either Selecto or Unselect Exact Tag depending on if we have the exact tag or not
+		// Either Selector or Unselect Exact Tag depending on if we have the exact tag or not
 		if (IsExactTagInCollection(InTagNode))
 		{
 			MenuBuilder.AddMenuEntry(LOCTEXT("GameplayTagPicker_UnselectTag", "Unselect Exact Tag"),
@@ -1201,6 +1291,15 @@ TSharedRef<SWidget> SGameplayTagPicker::MakeTagActionsMenu(TSharedPtr<FGameplayT
 		MenuBuilder.AddSeparator();
 	}
 
+	if (!bShowManagement)
+	{
+		// Open tag in manager
+		MenuBuilder.AddMenuEntry(LOCTEXT("GameplayTagPicker_OpenTagInManager", "Open Tag in Manager..."),
+			LOCTEXT("GameplayTagPicker_OpenTagInManagerTooltip", "Opens the Gameplay Tag manage and hilights the selected tag."),
+			FSlateIcon(FAppStyle::GetAppStyleSetName(), "Icons.Settings"),
+			FUIAction(FExecuteAction::CreateRaw(this, &SGameplayTagPicker::OnManageTagsClicked, InTagNode, TSharedPtr<SComboButton>())));
+	}
+
 	// Search for References
 	if (FEditorDelegates::OnOpenReferenceViewer.IsBound())
 	{
@@ -1234,26 +1333,50 @@ bool SGameplayTagPicker::CanModifyTag(TSharedPtr<FGameplayTagNode> Node) const
 
 void SGameplayTagPicker::OnAddSubTag(TSharedPtr<FGameplayTagNode> InTagNode, TSharedPtr<SComboButton> OwnerCombo)
 {
-	if (InTagNode.IsValid())
+	const bool bShowInline = ((GameplayTagPickerMode == EGameplayTagPickerMode::SelectionMode || GameplayTagPickerMode == EGameplayTagPickerMode::HybridMode) && !bReadOnly);
+
+	if (bShowInline)
 	{
-		AddTagDialog(EGameplayTagAdd::Child, InTagNode);
+		if (InTagNode.IsValid())
+		{
+			ShowInlineAddTagWidget(EGameplayTagAdd::Child, InTagNode);
+		}
 	}
-	if (OwnerCombo.IsValid())
+	else
 	{
-		OwnerCombo->SetIsOpen(false);
+		if (InTagNode.IsValid())
+		{
+			OpenAddTagDialog(EGameplayTagAdd::Child, InTagNode);
+		}
+		if (OwnerCombo.IsValid())
+		{
+			OwnerCombo->SetIsOpen(false);
+		}
 	}
 }
 
 void SGameplayTagPicker::OnDuplicateTag(TSharedPtr<FGameplayTagNode> InTagNode, TSharedPtr<SComboButton> OwnerCombo)
 {
-	if (InTagNode.IsValid())
+	const bool bShowInline = ((GameplayTagPickerMode == EGameplayTagPickerMode::SelectionMode || GameplayTagPickerMode == EGameplayTagPickerMode::HybridMode) && !bReadOnly);
+
+	if (bShowInline)
 	{
-		AddTagDialog(EGameplayTagAdd::Duplicate, InTagNode);
+		if (InTagNode.IsValid())
+		{
+			ShowInlineAddTagWidget(EGameplayTagAdd::Duplicate, InTagNode);
+		}
 	}
-	
-	if (OwnerCombo.IsValid())
+	else
 	{
-		OwnerCombo->SetIsOpen(false);
+		if (InTagNode.IsValid())
+		{
+			OpenAddTagDialog(EGameplayTagAdd::Duplicate, InTagNode);
+		}
+		
+		if (OwnerCombo.IsValid())
+		{
+			OwnerCombo->SetIsOpen(false);
+		}
 	}
 }
 
@@ -1677,8 +1800,7 @@ void SGameplayTagPicker::OpenRenameGameplayTagDialog(TSharedPtr<FGameplayTagNode
 	TSharedRef<SWindow> RenameTagWindow =
 		SNew(SWindow)
 		.Title(LOCTEXT("RenameTagWindowTitle", "Rename Gameplay Tag"))
-		.ClientSize(FVector2D(320.0f, 110.0f))
-		.SizingRule(ESizingRule::FixedSize)
+		.SizingRule(ESizingRule::Autosized)
 		.SupportsMaximize(false)
 		.SupportsMinimize(false);
 
@@ -1687,7 +1809,11 @@ void SGameplayTagPicker::OpenRenameGameplayTagDialog(TSharedPtr<FGameplayTagNode
 		.GameplayTagNode(GameplayTagNode)
 		.OnGameplayTagRenamed(const_cast<SGameplayTagPicker*>(this), &SGameplayTagPicker::OnGameplayTagRenamed);
 
-	RenameTagWindow->SetContent(RenameTagDialog);
+	RenameTagWindow->SetContent(SNew(SBox)
+		.MinDesiredWidth(320.0f)
+		[
+			RenameTagDialog
+		]);
 
 	TSharedPtr<SWindow> CurrentWindow = FSlateApplication::Get().FindWidgetWindow(AsShared() );
 
@@ -1805,67 +1931,80 @@ void CloseGameplayTagWindow(TWeakPtr<SGameplayTagPicker> TagWidget)
 
 TWeakPtr<SGameplayTagPicker> OpenGameplayTagManager(const FGameplayTagManagerWindowArgs& Args)
 {
-	// Close all other GameplayTag windows.
-	CloseGameplayTagWindow(nullptr);
-
-	const FVector2D WindowSize(800, 800);
+	TSharedPtr<SWindow> GameplayTagPickerWindow = GlobalTagWidgetWindow.Pin();
+	TSharedPtr<SGameplayTagPicker> TagWidget = GlobalTagWidget.Pin();
 	
-	TSharedRef<SGameplayTagPicker> TagWidget = SNew(SGameplayTagPicker)
-		.Filter(Args.Filter)
-		.ReadOnly(false)
-		.MaxHeight(0.0f) // unbounded
-		.MultiSelect(false)
-		.SettingsName(TEXT("Manager"))
-		.GameplayTagPickerMode(EGameplayTagPickerMode::ManagementMode)
-		.RestrictedTags(Args.bRestrictedTags)
-	;
-
-	FText Title = Args.Title;
-	if (Title.IsEmpty())
+	if (!GlobalTagWidgetWindow.IsValid()
+		|| !GlobalTagWidget.IsValid())
 	{
-		Title = LOCTEXT("GameplayTagPicker_ManagerTitle", "Gameplay Tag Manager");
-	}
-	
-	TSharedRef<SWindow> GameplayTagPickerWindow = SNew(SWindow)
-		.Title(Title)
-		.HasCloseButton(true)
-		.SupportsMaximize(false)
-		.SupportsMinimize(false)
-		.AutoCenter(EAutoCenter::PreferredWorkArea)
-		.ClientSize(WindowSize)
-		[
-			SNew(SBorder)
-			.BorderImage(FAppStyle::GetBrush("ToolPanel.GroupBorder"))
-			[
-				SNew(SVerticalBox)
+		// Close all other GameplayTag windows.
+		CloseGameplayTagWindow(nullptr);
 
-				+ SVerticalBox::Slot()
-				.FillHeight(1)
+		const FVector2D WindowSize(800, 800);
+		
+		TagWidget = SNew(SGameplayTagPicker)
+			.Filter(Args.Filter)
+			.ReadOnly(false)
+			.MaxHeight(0.0f) // unbounded
+			.MultiSelect(false)
+			.SettingsName(TEXT("Manager"))
+			.GameplayTagPickerMode(EGameplayTagPickerMode::ManagementMode)
+			.RestrictedTags(Args.bRestrictedTags)
+		;
+
+		FText Title = Args.Title;
+		if (Title.IsEmpty())
+		{
+			Title = LOCTEXT("GameplayTagPicker_ManagerTitle", "Gameplay Tag Manager");
+		}
+		
+		GameplayTagPickerWindow = SNew(SWindow)
+			.Title(Title)
+			.HasCloseButton(true)
+			.SupportsMaximize(false)
+			.SupportsMinimize(false)
+			.AutoCenter(EAutoCenter::PreferredWorkArea)
+			.ClientSize(WindowSize)
+			[
+				SNew(SBorder)
+				.BorderImage(FAppStyle::GetBrush("ToolPanel.GroupBorder"))
 				[
-					TagWidget
+					SNew(SVerticalBox)
+
+					+ SVerticalBox::Slot()
+					.FillHeight(1)
+					[
+						TagWidget.ToSharedRef()
+					]
 				]
-			]
-		];
+			];
 
-	TWeakPtr<SGameplayTagPicker> WeakTagWidget = TagWidget;
-	
-	// NOTE: FGlobalTabmanager::Get()-> is actually dereferencing a SharedReference, not a SharedPtr, so it cannot be null.
-	if (FGlobalTabmanager::Get()->GetRootWindow().IsValid())
-	{
-		FSlateApplication::Get().AddWindowAsNativeChild(GameplayTagPickerWindow, FGlobalTabmanager::Get()->GetRootWindow().ToSharedRef());
-	}
-	else
-	{
-		FSlateApplication::Get().AddWindow(GameplayTagPickerWindow);
+		TWeakPtr<SGameplayTagPicker> WeakTagWidget = TagWidget;
+		
+		// NOTE: FGlobalTabmanager::Get()-> is actually dereferencing a SharedReference, not a SharedPtr, so it cannot be null.
+		if (FGlobalTabmanager::Get()->GetRootWindow().IsValid())
+		{
+			FSlateApplication::Get().AddWindowAsNativeChild(GameplayTagPickerWindow.ToSharedRef(), FGlobalTabmanager::Get()->GetRootWindow().ToSharedRef());
+		}
+		else
+		{
+			FSlateApplication::Get().AddWindow(GameplayTagPickerWindow.ToSharedRef());
+		}
+
+		GlobalTagWidget = TagWidget;
+		GlobalTagWidgetWindow = GameplayTagPickerWindow;
 	}
 
-	
+	check (TagWidget.IsValid());
+
 	// Set focus to the search box on creation
 	FSlateApplication::Get().SetKeyboardFocus(TagWidget->GetWidgetToFocusOnOpen());
 	FSlateApplication::Get().SetUserFocus(0, TagWidget->GetWidgetToFocusOnOpen());
-	
-	GlobalTagWidget = TagWidget;
-	GlobalTagWidgetWindow = GameplayTagPickerWindow;
+
+	if (Args.HighlightedTag.IsValid())
+	{
+		TagWidget->RequestScrollToView(Args.HighlightedTag);
+	}
 
 	return TagWidget;
 }
diff --git a/Engine/Plugins/Editor/GameplayTagsEditor/Source/GameplayTagsEditor/Private/SRenameGameplayTagDialog.cpp b/Engine/Plugins/Editor/GameplayTagsEditor/Source/GameplayTagsEditor/Private/SRenameGameplayTagDialog.cpp
--- a/Engine/Plugins/Editor/GameplayTagsEditor/Source/GameplayTagsEditor/Private/SRenameGameplayTagDialog.cpp
+++ b/Engine/Plugins/Editor/GameplayTagsEditor/Source/GameplayTagsEditor/Private/SRenameGameplayTagDialog.cpp
@@ -9,6 +9,7 @@
 #include "Framework/Application/SlateApplication.h"
 #include "Widgets/SWindow.h"
 #include "Widgets/Layout/SBox.h"
+#include "Widgets/Layout/SGridPanel.h"
 
 #define LOCTEXT_NAMESPACE "RenameGameplayTag"
 
@@ -21,101 +22,88 @@ void SRenameGameplayTagDialog::Construct(const FArguments& InArgs)
 
 	ChildSlot
 	[
-		SNew( SBox )
-		.Padding(FMargin(15))
+		SNew(SBox)
+		.Padding(InArgs._Padding)
 		[
-			SNew( SVerticalBox )
+			SNew(SGridPanel)
+			.FillColumn(1, 1.0)
 
 			// Current name display
-			+SVerticalBox::Slot()
-			.AutoHeight()
-			.VAlign(VAlign_Top)
-			.Padding(4.0f)
+			+ SGridPanel::Slot(0, 0)
+			.Padding(2)
+			.VAlign(VAlign_Center)
+			.HAlign(HAlign_Left)
 			[
-				SNew( SHorizontalBox )
-
-				+SHorizontalBox::Slot()
-				.AutoWidth()
-				[
-					SNew( STextBlock )
-					.Text( LOCTEXT("CurrentTag", "Current Tag:"))
-				]
-
-				+SHorizontalBox::Slot()
-				.FillWidth(1.0f)
-				.HAlign(HAlign_Right)
-				.Padding(8.0f, 0.0f)
-				[
-					SNew( STextBlock )
-					.MinDesiredWidth(184.0f)
-					.Text( FText::FromName(GameplayTagNode->GetCompleteTag().GetTagName() ) )
-				]
+				SNew(STextBlock)
+				.Font(FAppStyle::GetFontStyle( TEXT("PropertyWindow.NormalFont")))
+				.Text(LOCTEXT("CurrentTag", "Current Tag:"))
 			]
 
+			+ SGridPanel::Slot(1, 0)
+			.Padding(2)
+			.VAlign(VAlign_Center)
+			.HAlign(HAlign_Fill)
+			[
+				SNew(STextBlock)
+				.Font(FAppStyle::GetFontStyle( TEXT("PropertyWindow.NormalFont")))
+				.Text(FText::FromName(GameplayTagNode->GetCompleteTag().GetTagName()))
+			]
 			
 			// New name controls
-			+SVerticalBox::Slot()
-			.AutoHeight()
-			.Padding(4.0f)
-			.VAlign(VAlign_Top)
+			+ SGridPanel::Slot(0, 1)
+			.Padding(2)
+			.VAlign(VAlign_Center)
+			.HAlign(HAlign_Left)
 			[
-				SNew( SHorizontalBox )
-
-				+SHorizontalBox::Slot()
-				.AutoWidth()
-				.Padding(0.0f, 4.0f)
-				[
-					SNew( STextBlock )
-					.Text( LOCTEXT("NewTag", "New Tag:" ))
-				]
+				SNew(STextBlock )
+				.Font(FAppStyle::GetFontStyle( TEXT("PropertyWindow.NormalFont")))
+				.Text(LOCTEXT("NewTag", "New Tag:"))
+			]
 
-				+SHorizontalBox::Slot()
-				.FillWidth(1.0f)
-				.HAlign(HAlign_Right)
-				.Padding(8.0f, 0.0f)
-				[
-					SAssignNew( NewTagNameTextBox, SEditableTextBox )
-					.Text( FText::FromName(GameplayTagNode->GetCompleteTag().GetTagName() ))
-					.Padding(4.0f)
-					.MinDesiredWidth(180.0f)
-					.OnTextCommitted( this, &SRenameGameplayTagDialog::OnRenameTextCommitted )
-				]
+			+ SGridPanel::Slot(1, 1)
+			.Padding(2)
+			.VAlign(VAlign_Center)
+			.HAlign(HAlign_Fill)
+			[
+				SAssignNew(NewTagNameTextBox, SEditableTextBox)
+				.Font(FAppStyle::GetFontStyle( TEXT("PropertyWindow.NormalFont")))
+				.Text(FText::FromName(GameplayTagNode->GetCompleteTag().GetTagName()))
+				.OnTextCommitted(this, &SRenameGameplayTagDialog::OnRenameTextCommitted)
 			]
 
 			// Dialog controls
-			+SVerticalBox::Slot()
-			.AutoHeight()
-			.VAlign(VAlign_Top)
-			.HAlign(HAlign_Center)
+			+ SGridPanel::Slot(0, 2)
+			.ColumnSpan(2)
+			.HAlign(HAlign_Right)
+			.Padding(FMargin(0, 16))
 			[
-				SNew( SHorizontalBox )
+				SNew(SHorizontalBox)
 
 				// Rename
 				+SHorizontalBox::Slot()
 				.AutoWidth()
-				.Padding(8.0f, 8.0f)
+				.Padding(0, 0, 8.0f, 0)
 				[
-					SNew( SButton )
-					.IsFocusable( false )
-					.IsEnabled( this, &SRenameGameplayTagDialog::IsRenameEnabled )
-					.OnClicked( this, &SRenameGameplayTagDialog::OnRenameClicked )
+					SNew(SButton)
+					.IsFocusable(false)
+					.IsEnabled(this, &SRenameGameplayTagDialog::IsRenameEnabled)
+					.OnClicked(this, &SRenameGameplayTagDialog::OnRenameClicked)
 					[
-						SNew( STextBlock )
-						.Text( LOCTEXT("RenameTagButtonText", "Rename" ) )
+						SNew(STextBlock)
+						.Text(LOCTEXT("RenameTagButtonText", "Rename"))
 					]
 				]
 
 				// Cancel
 				+SHorizontalBox::Slot()
 				.AutoWidth()
-				.Padding(8.0f, 8.0f)
 				[
-					SNew( SButton )
-					.IsFocusable( false )
-					.OnClicked( this, & SRenameGameplayTagDialog::OnCancelClicked )
+					SNew(SButton)
+					.IsFocusable(false)
+					.OnClicked(this, & SRenameGameplayTagDialog::OnCancelClicked)
 					[
-						SNew( STextBlock )
-						.Text( LOCTEXT("CancelRenameButtonText", "Cancel"))
+						SNew(STextBlock)
+						.Text(LOCTEXT("CancelRenameButtonText", "Cancel"))
 					]
 				]
 			]
diff --git a/Engine/Plugins/Editor/GameplayTagsEditor/Source/GameplayTagsEditor/Private/SRenameGameplayTagDialog.h b/Engine/Plugins/Editor/GameplayTagsEditor/Source/GameplayTagsEditor/Private/SRenameGameplayTagDialog.h
--- a/Engine/Plugins/Editor/GameplayTagsEditor/Source/GameplayTagsEditor/Private/SRenameGameplayTagDialog.h
+++ b/Engine/Plugins/Editor/GameplayTagsEditor/Source/GameplayTagsEditor/Private/SRenameGameplayTagDialog.h
@@ -16,9 +16,11 @@ public:
 
 	SLATE_BEGIN_ARGS( SRenameGameplayTagDialog )
 		: _GameplayTagNode()
+		, _Padding(FMargin(15))
 		{}
 		SLATE_ARGUMENT( TSharedPtr<FGameplayTagNode>, GameplayTagNode )		// The gameplay tag we want to rename
 		SLATE_EVENT( FOnGameplayTagRenamed, OnGameplayTagRenamed )	// Called when the tag is renamed
+		SLATE_ARGUMENT(FMargin, Padding)
 	SLATE_END_ARGS()
 
 	void Construct(const FArguments& InArgs);
diff --git a/Engine/Plugins/Editor/GameplayTagsEditor/Source/GameplayTagsEditor/Public/SGameplayTagPicker.h b/Engine/Plugins/Editor/GameplayTagsEditor/Source/GameplayTagsEditor/Public/SGameplayTagPicker.h
--- a/Engine/Plugins/Editor/GameplayTagsEditor/Source/GameplayTagsEditor/Public/SGameplayTagPicker.h
+++ b/Engine/Plugins/Editor/GameplayTagsEditor/Source/GameplayTagsEditor/Public/SGameplayTagPicker.h
@@ -17,6 +17,7 @@
 
 class IPropertyHandle;
 class SComboButton;
+class SAddNewGameplayTagWidget;
 
 /** Determines the behavior of the gameplay tag UI depending on where it's used */
 enum class EGameplayTagPickerMode : uint8
@@ -260,7 +261,7 @@ private:
 	FSlateColor GetTagTextColour(TSharedPtr<FGameplayTagNode> Node) const;
 
 	/** Called when the user clicks the "Create and Manage Gameplay Tags" button; Opens Gameplay Tag manager window. */
-	void OnManageTagsClicked(TSharedPtr<SComboButton> OwnerCombo);
+	void OnManageTagsClicked(TSharedPtr<FGameplayTagNode> Node, TSharedPtr<SComboButton> OwnerCombo);
 
 	/** Called when the user clicks the "Clear All" button; Clears all tags */
 	void OnClearAllClicked(TSharedPtr<SComboButton> OwnerCombo);
@@ -318,10 +319,11 @@ private:
 	/** Callback for when a new tag is added */
 	void OnGameplayTagAdded(const FString& TagName, const FString& TagComment, const FName& TagSource);
 
-	/** Callback when the user wants to add a subtag to an existing tag */
-	FReply OnAddSubtagClicked(TSharedPtr<FGameplayTagNode> InTagNode);
+	/** Opens add tag modal dialog. */
+	void OpenAddTagDialog(const EGameplayTagAdd Mode, TSharedPtr<FGameplayTagNode> InTagNode = TSharedPtr<FGameplayTagNode>());
 
-	void AddTagDialog(const EGameplayTagAdd Mode, TSharedPtr<FGameplayTagNode> InTagNode = TSharedPtr<FGameplayTagNode>());
+	/** Initializes the inline add tag widget and makes in visible. */
+	void ShowInlineAddTagWidget(const EGameplayTagAdd Mode, TSharedPtr<FGameplayTagNode> InTagNode = TSharedPtr<FGameplayTagNode>());
 
 	FReply OnAddRootTagClicked();
 
@@ -333,9 +335,6 @@ private:
 	/** Creates add menu content. */
 	TSharedRef<SWidget> MakeAddMenu();
 
-	/** Creates actions menu content. */
-	TSharedRef<SWidget> MakeActionsMenu(TSharedPtr<SComboButton> OwnerCombo);
-
 	/** Creates settings menu content. */
 	TSharedRef<SWidget> MakeSettingsMenu(TSharedPtr<SComboButton> OwnerCombo);
 
@@ -358,11 +357,13 @@ private:
 	/** Searches for all references for the selected tag */
 	void OnSearchForReferences(TSharedPtr<FGameplayTagNode> InTagNode, TSharedPtr<SComboButton> OwnerCombo);
 
+	/** Copies individual tag's name to clipboard. */
 	void OnCopyTagNameToClipboard(TSharedPtr<FGameplayTagNode> InTagNode, TSharedPtr<SComboButton> OwnerCombo);
 	
 	/** Returns true if the user can select tags from the widget */
 	bool CanSelectTags() const;
 
+	/** Called when tags are selected. */
 	void OnContainersChanged();
 
 	/** Opens a dialog window to rename the selected tag */
@@ -451,6 +452,9 @@ private:
 
 	/** Expanded items cached from settings */
 	TSet<TSharedPtr<FGameplayTagNode>> CachedExpandedItems;
+
+	TSharedPtr<SAddNewGameplayTagWidget> AddNewTagWidget;
+	bool bNewTagWidgetVisible = false;
 };
 
 
@@ -458,6 +462,7 @@ struct FGameplayTagManagerWindowArgs
 {
 	FText Title;
 	FString Filter; // Comma delimited string of tag root names to filter by
+	FGameplayTag HighlightedTag; // Tag to highlight when window is opened. 
 	bool bRestrictedTags;
 };
 
