Fix Gameplay Tag Query editor when editing data tables
- Do not expect that we have always outer objects (not the case when editing only a struct)
- Show package name in the query editor if we dont have outer object

diff --git a/Engine/Plugins/Editor/GameplayTagsEditor/Source/GameplayTagsEditor/Private/SGameplayTagQueryEntryBox.cpp b/Engine/Plugins/Editor/GameplayTagsEditor/Source/GameplayTagsEditor/Private/SGameplayTagQueryEntryBox.cpp
--- a/Engine/Plugins/Editor/GameplayTagsEditor/Source/GameplayTagsEditor/Private/SGameplayTagQueryEntryBox.cpp
+++ b/Engine/Plugins/Editor/GameplayTagsEditor/Source/GameplayTagsEditor/Private/SGameplayTagQueryEntryBox.cpp
@@ -208,6 +208,17 @@ FReply SGameplayTagQueryEntryBox::OnEditButtonClicked()
 
 		TArray<UObject*> OuterObjects;
 		PropertyHandle->GetOuterObjects(OuterObjects);
+
+		if (OuterObjects.IsEmpty())
+		{
+			TArray<UPackage*> OuterPackages;
+			PropertyHandle->GetOuterPackages(OuterPackages);
+			for (UPackage* Package : OuterPackages)
+			{
+				OuterObjects.Add(Package);
+			}
+		}
+		
 		if (OuterObjects.Num() > 1)
 		{
 			FText const AssetName = FText::Format(LOCTEXT("GameplayTagDetailsBase_MultipleAssets", "{0} Assets"), FText::AsNumber(OuterObjects.Num()));
@@ -294,19 +305,13 @@ void SGameplayTagQueryEntryBox::OnQueriesCommitted(const TArray<FGameplayTagQuer
 
 		PropertyHandle->NotifyPreChange();
 
-		TArray<UObject*> OuterObjects;
-		PropertyHandle->GetOuterObjects(OuterObjects);
-
-		if (OuterObjects.Num() > 0 && OuterObjects.Num() == TagQueries.Num())
+		TArray<FString> PerObjectValues;
+		for (const FGameplayTagQuery& TagQuery : TagQueries)
 		{
-			TArray<FString> PerObjectValues;
-			for (const FGameplayTagQuery& TagQuery : TagQueries)
-			{
-				FString& ExportString = PerObjectValues.AddDefaulted_GetRef();
-				FGameplayTagQuery::StaticStruct()->ExportText(ExportString, &TagQuery, &TagQuery, nullptr, 0, nullptr);
-			}
-			PropertyHandle->SetPerObjectValues(PerObjectValues);
+			FString& ExportString = PerObjectValues.AddDefaulted_GetRef();
+			FGameplayTagQuery::StaticStruct()->ExportText(ExportString, &TagQuery, &TagQuery, nullptr, 0, nullptr);
 		}
+		PropertyHandle->SetPerObjectValues(PerObjectValues);
 
 		PropertyHandle->NotifyPostChange(EPropertyChangeType::ValueSet);
 		PropertyHandle->NotifyFinishedChangingProperties();
