Sometimes GameplayTag edit widgets don't update in the quest editor when their edit condition changes

diff --git a/Engine/Plugins/Editor/GameplayTagsEditor/Source/GameplayTagsEditor/Private/SGameplayTagCombo.cpp b/Engine/Plugins/Editor/GameplayTagsEditor/Source/GameplayTagsEditor/Private/SGameplayTagCombo.cpp
--- a/Engine/Plugins/Editor/GameplayTagsEditor/Source/GameplayTagsEditor/Private/SGameplayTagCombo.cpp
+++ b/Engine/Plugins/Editor/GameplayTagsEditor/Source/GameplayTagsEditor/Private/SGameplayTagCombo.cpp
@@ -71,7 +71,7 @@ void SGameplayTagCombo::Construct(const FArguments& InArgs)
 			.ComboButtonStyle(FGameplayTagStyle::Get(), "GameplayTags.ComboButton")
 			.HasDownArrow(true)
 			.ContentPadding(1)
-			.IsEnabled(!bIsReadOnly)
+			.IsEnabled(this, &SGameplayTagCombo::IsValueEnabled)
 			.Clipping(EWidgetClipping::OnDemand)
 			.OnMenuOpenChanged(this, &SGameplayTagCombo::OnMenuOpenChanged)
 			.OnGetMenuContent(this, &SGameplayTagCombo::OnGetMenuContent)
@@ -93,6 +93,16 @@ void SGameplayTagCombo::Construct(const FArguments& InArgs)
 	];
 }
 
+bool SGameplayTagCombo::IsValueEnabled() const
+{
+	if (PropertyHandle.IsValid())
+	{
+		return !PropertyHandle->IsEditConst();
+	}
+
+	return bIsReadOnly;
+}
+
 void SGameplayTagCombo::PostUndo(bool bSuccess)
 {
 	if (bSuccess)
@@ -214,7 +224,7 @@ TSharedRef<SWidget> SGameplayTagCombo::OnGetMenuContent()
 	TagPicker = SNew(SGameplayTagPicker)
 		.Filter(Filter)
 		.SettingsName(SettingsName)
-		.ReadOnly(bIsReadOnly)
+		.ReadOnly(IsValueEnabled())
 		.ShowMenuItems(true)
 		.MaxHeight(350.0f)
 		.MultiSelect(false)
diff --git a/Engine/Plugins/Editor/GameplayTagsEditor/Source/GameplayTagsEditor/Public/SGameplayTagCombo.h b/Engine/Plugins/Editor/GameplayTagsEditor/Source/GameplayTagsEditor/Public/SGameplayTagCombo.h
--- a/Engine/Plugins/Editor/GameplayTagsEditor/Source/GameplayTagsEditor/Public/SGameplayTagCombo.h
+++ b/Engine/Plugins/Editor/GameplayTagsEditor/Source/GameplayTagsEditor/Public/SGameplayTagCombo.h
@@ -76,6 +76,7 @@ private:
 
 	bool ShowClearButton() const;
 	FText GetText() const;
+	bool IsValueEnabled() const;
 	FText GetToolTipText() const;
 	bool IsSelected() const;
 	FReply OnClearPressed();
