Fix for FGameplayTagQuery not showing a proper Query Description if the struct was created by native code.

diff --git a/Engine/Plugins/Editor/GameplayTagsEditor/Source/GameplayTagsEditor/Private/SGameplayTagQueryEntryBox.cpp b/Engine/Plugins/Editor/GameplayTagsEditor/Source/GameplayTagsEditor/Private/SGameplayTagQueryEntryBox.cpp
--- a/Engine/Plugins/Editor/GameplayTagsEditor/Source/GameplayTagsEditor/Private/SGameplayTagQueryEntryBox.cpp
+++ b/Engine/Plugins/Editor/GameplayTagsEditor/Source/GameplayTagsEditor/Private/SGameplayTagQueryEntryBox.cpp
@@ -1,4 +1,4 @@
-﻿// Copyright Epic Games, Inc. All Rights Reserved.
+// Copyright Epic Games, Inc. All Rights Reserved.
 
 #include "SGameplayTagQueryEntryBox.h"
 #include "DetailLayoutBuilder.h"
@@ -287,7 +287,20 @@ void SGameplayTagQueryEntryBox::CacheQueryList()
 	}
 	else if (CachedQueries.Num() == 1)
 	{
-		const FString Desc = CachedQueries[0].GetDescription();
+		const FGameplayTagQuery& TheQuery = CachedQueries[0];
+		FString Desc = TheQuery.GetDescription();
+
+		// We can get into a case where we've manually generated an FGameplayTagQuery but do not have a proper description
+		// in this case, regenerate a description for the query and use that until next time it's edited.
+		if (Desc.IsEmpty() && !TheQuery.IsEmpty())
+		{
+			UEditableGameplayTagQuery* EditableQuery = TheQuery.CreateEditableQuery();
+
+			FGameplayTagQuery TempQueryForDescription;
+			TempQueryForDescription.BuildFromEditableQuery(*EditableQuery);
+			Desc = TempQueryForDescription.GetDescription();
+		}
+
 		if (Desc.Len() > 0)
 		{
 			QueryDescription = FText::FromString(Desc);
diff --git a/Engine/Source/Runtime/GameplayTags/Classes/GameplayTagContainer.h b/Engine/Source/Runtime/GameplayTags/Classes/GameplayTagContainer.h
--- a/Engine/Source/Runtime/GameplayTags/Classes/GameplayTagContainer.h
+++ b/Engine/Source/Runtime/GameplayTags/Classes/GameplayTagContainer.h
@@ -893,7 +893,7 @@ public:
 	void BuildFromEditableQuery(class UEditableGameplayTagQuery& EditableQuery); 
 
 	/** Creates editable query object tree based on this query */
-	UEditableGameplayTagQuery* CreateEditableQuery();
+	UEditableGameplayTagQuery* CreateEditableQuery() const;
 #endif // WITH_EDITOR
 
 	static const FGameplayTagQuery EmptyQuery;
diff --git a/Engine/Source/Runtime/GameplayTags/Private/GameplayTagContainer.cpp b/Engine/Source/Runtime/GameplayTags/Private/GameplayTagContainer.cpp
--- a/Engine/Source/Runtime/GameplayTags/Private/GameplayTagContainer.cpp
+++ b/Engine/Source/Runtime/GameplayTags/Private/GameplayTagContainer.cpp
@@ -1660,7 +1660,7 @@ void FQueryEvaluator::ReadEditableQueryExprList(UEditableGameplayTagQueryExpress
 	}
 }
 
-UEditableGameplayTagQuery* FGameplayTagQuery::CreateEditableQuery()
+UEditableGameplayTagQuery* FGameplayTagQuery::CreateEditableQuery() const
 {
 	FQueryEvaluator QE(*this);
 	return QE.CreateEditableQuery();
