GameplayTags are not choosable in Make Array in Blueprints

diff --git a/Engine/Plugins/Editor/GameplayTagsEditor/Source/GameplayTagsEditor/Private/SGameplayTagCombo.cpp b/Engine/Plugins/Editor/GameplayTagsEditor/Source/GameplayTagsEditor/Private/SGameplayTagCombo.cpp
--- a/Engine/Plugins/Editor/GameplayTagsEditor/Source/GameplayTagsEditor/Private/SGameplayTagCombo.cpp
+++ b/Engine/Plugins/Editor/GameplayTagsEditor/Source/GameplayTagsEditor/Private/SGameplayTagCombo.cpp
@@ -100,7 +100,7 @@ bool SGameplayTagCombo::IsValueEnabled() const
 		return !PropertyHandle->IsEditConst();
 	}
 
-	return bIsReadOnly;
+	return !bIsReadOnly;
 }
 
 void SGameplayTagCombo::PostUndo(bool bSuccess)
