GameplayTags: Use inline tag add widget for tag manager too.

diff --git a/Engine/Plugins/Editor/GameplayTagsEditor/Source/GameplayTagsEditor/Private/SGameplayTagPicker.cpp b/Engine/Plugins/Editor/GameplayTagsEditor/Source/GameplayTagsEditor/Private/SGameplayTagPicker.cpp
--- a/Engine/Plugins/Editor/GameplayTagsEditor/Source/GameplayTagsEditor/Private/SGameplayTagPicker.cpp
+++ b/Engine/Plugins/Editor/GameplayTagsEditor/Source/GameplayTagsEditor/Private/SGameplayTagPicker.cpp
@@ -25,7 +25,6 @@
 #include "Framework/MultiBox/MultiBoxBuilder.h"
 #include "Framework/Docking/TabManager.h"
 #include "HAL/PlatformApplicationMisc.h"
-#include "SPositiveActionButton.h"
 #include "Misc/EnumerateRange.h"
 #include "Styling/StyleColors.h"
 
@@ -165,25 +164,6 @@ void SGameplayTagPicker::Construct(const FArguments& InArgs)
 			[
 				SNew(SHorizontalBox)
 
-				// Add button for management mode
-				+SHorizontalBox::Slot()
-				.AutoWidth()
-				[
-					SNew(SHorizontalBox)
-
-					+SHorizontalBox::Slot()
-					.AutoWidth()
-					.Padding(FMargin(0,0,4,0))
-					[
-						SNew(SPositiveActionButton)
-						.Icon(FAppStyle::Get().GetBrush("Icons.Plus"))
-						.Text(LOCTEXT("GameplayTagPicker_Add", "Add"))
-						.IsEnabled(this, &SGameplayTagPicker::CanAddNewTag)
-						.Visibility(InArgs._GameplayTagPickerMode == EGameplayTagPickerMode::ManagementMode ? EVisibility::Visible : EVisibility::Collapsed)
-						.OnGetMenuContent(this, &SGameplayTagPicker::MakeAddMenu)
-					]
-				]
-
 				// Smaller add button for selection and hybrid modes.
 				+SHorizontalBox::Slot()
 				.AutoWidth()
@@ -196,14 +176,31 @@ void SGameplayTagPicker::Construct(const FArguments& InArgs)
 					[
 						SNew(SButton)
 						.ButtonStyle(FAppStyle::Get(), "HoverHintOnly")
-						.Visibility(InArgs._GameplayTagPickerMode != EGameplayTagPickerMode::ManagementMode ? EVisibility::Visible : EVisibility::Collapsed)
+						.ToolTipText_Lambda([WeakSelf]() -> FText
+						{
+							const TSharedPtr<SGameplayTagPicker> Self = WeakSelf.Pin();
+							if (Self.IsValid() && Self->bNewTagWidgetVisible)
+							{
+								return LOCTEXT("CloseSection", "Close Section");
+							}
+							return LOCTEXT("AddNewGameplayTag", "Add New Gameplay Tag");
+						})
 						.OnClicked_Lambda([WeakSelf]()
 						{
 							if (const TSharedPtr<SGameplayTagPicker> Self = WeakSelf.Pin())
 							{
 								if (!Self->bNewTagWidgetVisible)
 								{
-									Self->ShowInlineAddTagWidget(EGameplayTagAdd::Root);
+									// If we have a selected item, by default add child, else new root tag. 
+									TArray<TSharedPtr<FGameplayTagNode>> Selection = Self->TagTreeWidget->GetSelectedItems();
+									if (Selection.Num() > 0 && Selection[0].IsValid())
+									{
+										Self->ShowInlineAddTagWidget(EGameplayTagAdd::Child, Selection[0]);
+									}
+									else
+									{
+										Self->ShowInlineAddTagWidget(EGameplayTagAdd::Root);
+									}
 								}
 								else
 								{
@@ -252,7 +249,6 @@ void SGameplayTagPicker::Construct(const FArguments& InArgs)
 				{
 					const TSharedPtr<SGameplayTagPicker> Self = WeakSelf.Pin();
 					return (Self.IsValid() && Self->bNewTagWidgetVisible) ? EVisibility::Visible : EVisibility::Collapsed;
-					
 				})
 				.OnGameplayTagAdded_Lambda([WeakSelf](const FString& TagName, const FString& TagComment, const FName& TagSource)
 				{
@@ -407,35 +403,6 @@ FReply SGameplayTagPicker::OnTreeKeyDown(const FGeometry& InGeometry, const FKey
 	return SCompoundWidget::OnKeyDown(InGeometry, InKeyEvent);
 }
 
-TSharedRef<SWidget> SGameplayTagPicker::MakeAddMenu()
-{
-	constexpr bool bInShouldCloseWindowAfterMenuSelection = true;
-	FMenuBuilder MenuBuilder(bInShouldCloseWindowAfterMenuSelection, nullptr);
-
-	TArray<TSharedPtr<FGameplayTagNode>> Selection = TagTreeWidget->GetSelectedItems();
-	const TSharedPtr<FGameplayTagNode> SelectedTagNode = Selection.IsEmpty() ? nullptr : Selection[0];
-
-	// Add root tag
-	MenuBuilder.AddMenuEntry(
-		LOCTEXT("GameplayTagPicker_AddNewRootTag", "Add new root level tag"), FText::GetEmpty(), FSlateIcon(FAppStyle::GetAppStyleSetName(), "Icons.Plus"),
-		FUIAction(FExecuteAction::CreateRaw(this, &SGameplayTagPicker::OpenAddTagDialog, EGameplayTagAdd::Root, TSharedPtr<FGameplayTagNode>()))
-	);
-
-	// Add child tag
-	MenuBuilder.AddMenuEntry(LOCTEXT("GameplayTagPicker_AddSubTag", "Add Sub Tag"),
-		LOCTEXT("GameplayTagPicker_AddSubTagTagTooltip", "Add sub tag under selected tag."),
-		FSlateIcon(FAppStyle::GetAppStyleSetName(), "Icons.Plus"),
-		FUIAction(FExecuteAction::CreateSP(this, &SGameplayTagPicker::OnAddSubTag, SelectedTagNode, TSharedPtr<SComboButton>()), FCanExecuteAction::CreateSP(this, &SGameplayTagPicker::CanAddNewSubTag, SelectedTagNode)));
-
-	// Duplicate
-	MenuBuilder.AddMenuEntry(LOCTEXT("GameplayTagPicker_DuplicateTag", "Duplicate Tag"),
-		LOCTEXT("GameplayTagPicker_DuplicateTagTooltip", "Duplicate selected tag to create a new tag."),
-		FSlateIcon(FAppStyle::GetAppStyleSetName(), "Icons.Duplicate"),
-		FUIAction(FExecuteAction::CreateSP(this, &SGameplayTagPicker::OnDuplicateTag, SelectedTagNode, TSharedPtr<SComboButton>()), FCanExecuteAction::CreateSP(this, &SGameplayTagPicker::CanAddNewSubTag, SelectedTagNode)));
-	
-	return MenuBuilder.MakeWidget();
-}
-
 TSharedRef<SWidget> SGameplayTagPicker::MakeSettingsMenu(TSharedPtr<SComboButton> OwnerCombo)
 {
 	FMenuBuilder MenuBuilder(/*bInShouldCloseWindowAfterMenuSelection*/false, nullptr);
@@ -486,6 +453,7 @@ void SGameplayTagPicker::Tick(const FGeometry& AllottedGeometry, const double In
 				ParentNode = ParentNode->ParentNode;
 			}
 
+			TagTreeWidget->ClearSelection();
 			TagTreeWidget->SetItemSelection(Node, true);
 			TagTreeWidget->RequestScrollIntoView(Node);
 		}
@@ -1333,50 +1301,27 @@ bool SGameplayTagPicker::CanModifyTag(TSharedPtr<FGameplayTagNode> Node) const
 
 void SGameplayTagPicker::OnAddSubTag(TSharedPtr<FGameplayTagNode> InTagNode, TSharedPtr<SComboButton> OwnerCombo)
 {
-	const bool bShowInline = ((GameplayTagPickerMode == EGameplayTagPickerMode::SelectionMode || GameplayTagPickerMode == EGameplayTagPickerMode::HybridMode) && !bReadOnly);
-
-	if (bShowInline)
+	if (InTagNode.IsValid())
 	{
-		if (InTagNode.IsValid())
-		{
-			ShowInlineAddTagWidget(EGameplayTagAdd::Child, InTagNode);
-		}
+		ShowInlineAddTagWidget(EGameplayTagAdd::Child, InTagNode);
 	}
-	else
+
+	if (OwnerCombo.IsValid())
 	{
-		if (InTagNode.IsValid())
-		{
-			OpenAddTagDialog(EGameplayTagAdd::Child, InTagNode);
-		}
-		if (OwnerCombo.IsValid())
-		{
-			OwnerCombo->SetIsOpen(false);
-		}
+		OwnerCombo->SetIsOpen(false);
 	}
 }
 
 void SGameplayTagPicker::OnDuplicateTag(TSharedPtr<FGameplayTagNode> InTagNode, TSharedPtr<SComboButton> OwnerCombo)
 {
-	const bool bShowInline = ((GameplayTagPickerMode == EGameplayTagPickerMode::SelectionMode || GameplayTagPickerMode == EGameplayTagPickerMode::HybridMode) && !bReadOnly);
-
-	if (bShowInline)
+	if (InTagNode.IsValid())
 	{
-		if (InTagNode.IsValid())
-		{
-			ShowInlineAddTagWidget(EGameplayTagAdd::Duplicate, InTagNode);
-		}
+		ShowInlineAddTagWidget(EGameplayTagAdd::Duplicate, InTagNode);
 	}
-	else
+	
+	if (OwnerCombo.IsValid())
 	{
-		if (InTagNode.IsValid())
-		{
-			OpenAddTagDialog(EGameplayTagAdd::Duplicate, InTagNode);
-		}
-		
-		if (OwnerCombo.IsValid())
-		{
-			OwnerCombo->SetIsOpen(false);
-		}
+		OwnerCombo->SetIsOpen(false);
 	}
 }
 
diff --git a/Engine/Plugins/Editor/GameplayTagsEditor/Source/GameplayTagsEditor/Public/SGameplayTagPicker.h b/Engine/Plugins/Editor/GameplayTagsEditor/Source/GameplayTagsEditor/Public/SGameplayTagPicker.h
--- a/Engine/Plugins/Editor/GameplayTagsEditor/Source/GameplayTagsEditor/Public/SGameplayTagPicker.h
+++ b/Engine/Plugins/Editor/GameplayTagsEditor/Source/GameplayTagsEditor/Public/SGameplayTagPicker.h
@@ -332,9 +332,6 @@ private:
 	/** Creates a dropdown menu to provide additional functionality for tags (renaming, deletion, search for references, etc.) */
 	TSharedRef<SWidget> MakeTagActionsMenu(TSharedPtr<FGameplayTagNode> InTagNode, TSharedPtr<SComboButton> OwnerCombo, const bool bInShouldCloseWindowAfterMenuSelection);
 
-	/** Creates add menu content. */
-	TSharedRef<SWidget> MakeAddMenu();
-
 	/** Creates settings menu content. */
 	TSharedRef<SWidget> MakeSettingsMenu(TSharedPtr<SComboButton> OwnerCombo);
 
