GameplayTags: Fixed crash and errorneous readonly behavior
- Fixed SGameplayTagCombo tag picker using inverted value if read only
- Changed SGameplayTagContainerCombo to handle edit conditions
- Do not allow to add tags on read only container

diff --git a/Engine/Plugins/Editor/GameplayTagsEditor/Source/GameplayTagsEditor/Private/SGameplayTagCombo.cpp b/Engine/Plugins/Editor/GameplayTagsEditor/Source/GameplayTagsEditor/Private/SGameplayTagCombo.cpp
--- a/Engine/Plugins/Editor/GameplayTagsEditor/Source/GameplayTagsEditor/Private/SGameplayTagCombo.cpp
+++ b/Engine/Plugins/Editor/GameplayTagsEditor/Source/GameplayTagsEditor/Private/SGameplayTagCombo.cpp
@@ -221,10 +221,12 @@ TSharedRef<SWidget> SGameplayTagCombo::OnGetMenuContent()
 		TagContainers.Add(FGameplayTagContainer(TagToEdit));
 	}
 
+	const bool bIsPickerReadOnly = !IsValueEnabled();
+	
 	TagPicker = SNew(SGameplayTagPicker)
 		.Filter(Filter)
 		.SettingsName(SettingsName)
-		.ReadOnly(IsValueEnabled())
+		.ReadOnly(bIsPickerReadOnly)
 		.ShowMenuItems(true)
 		.MaxHeight(350.0f)
 		.MultiSelect(false)
diff --git a/Engine/Plugins/Editor/GameplayTagsEditor/Source/GameplayTagsEditor/Private/SGameplayTagContainerCombo.cpp b/Engine/Plugins/Editor/GameplayTagsEditor/Source/GameplayTagsEditor/Private/SGameplayTagContainerCombo.cpp
--- a/Engine/Plugins/Editor/GameplayTagsEditor/Source/GameplayTagsEditor/Private/SGameplayTagContainerCombo.cpp
+++ b/Engine/Plugins/Editor/GameplayTagsEditor/Source/GameplayTagsEditor/Private/SGameplayTagContainerCombo.cpp
@@ -97,7 +97,7 @@ void SGameplayTagContainerCombo::Construct(const FArguments& InArgs)
 		[
 			SAssignNew(ComboButton, SComboButton)
 			.ComboButtonStyle(FGameplayTagStyle::Get(), "GameplayTagsContainer.ComboButton")
-			.IsEnabled(!bIsReadOnly)
+			.IsEnabled(this, &SGameplayTagContainerCombo::IsValueEnabled)
 			.HasDownArrow(true)
 			.VAlign(VAlign_Top)
 			.ContentPadding(0)
@@ -158,6 +158,16 @@ void SGameplayTagContainerCombo::Construct(const FArguments& InArgs)
 	];
 }
 
+bool SGameplayTagContainerCombo::IsValueEnabled() const
+{
+	if (PropertyHandle.IsValid())
+	{
+		return !PropertyHandle->IsEditConst();
+	}
+
+	return !bIsReadOnly;
+}
+
 TSharedRef<ITableRow> SGameplayTagContainerCombo::MakeTagListViewRow(TSharedPtr<FEditableItem> Item, const TSharedRef<STableViewBase>& OwnerTable)
 {
 	return SNew(STableRow<TSharedPtr<FEditableItem>>, OwnerTable)
@@ -200,10 +210,12 @@ TSharedRef<SWidget> SGameplayTagContainerCombo::OnGetMenuContent()
 		TagContainersToEdit.Add(CachedTagContainers[0]);
 	}
 
+	const bool bIsPickerReadOnly = !IsValueEnabled();
+
 	TagPicker = SNew(SGameplayTagPicker)
 		.Filter(Filter)
 		.SettingsName(SettingsName)
-		.ReadOnly(bIsReadOnly)
+		.ReadOnly(bIsPickerReadOnly)
 		.ShowMenuItems(true)
 		.MaxHeight(400.0f)
 		.MultiSelect(true)
@@ -216,7 +228,7 @@ TSharedRef<SWidget> SGameplayTagContainerCombo::OnGetMenuContent()
 	{
 		ComboButton->SetMenuContentWidgetToFocus(TagPicker->GetWidgetToFocusOnOpen());
 	}
-				
+
 	return TagPicker.ToSharedRef();
 }
 
diff --git a/Engine/Plugins/Editor/GameplayTagsEditor/Source/GameplayTagsEditor/Private/SGameplayTagPicker.cpp b/Engine/Plugins/Editor/GameplayTagsEditor/Source/GameplayTagsEditor/Private/SGameplayTagPicker.cpp
--- a/Engine/Plugins/Editor/GameplayTagsEditor/Source/GameplayTagsEditor/Private/SGameplayTagPicker.cpp
+++ b/Engine/Plugins/Editor/GameplayTagsEditor/Source/GameplayTagsEditor/Private/SGameplayTagPicker.cpp
@@ -1135,7 +1135,7 @@ void SGameplayTagPicker::OpenAddTagDialog(const EGameplayTagAdd Mode, TSharedPtr
 				AddNewTagDialog
 			]);
     }
-
+	
 	TSharedPtr<SWindow> CurrentWindow = FSlateApplication::Get().FindWidgetWindow(AsShared());
 	FSlateApplication::Get().AddModalWindow(NewTagWindow.ToSharedRef(), CurrentWindow);
 }
@@ -1687,7 +1687,7 @@ EVisibility SGameplayTagPicker::DetermineExpandableUIVisibility() const
 bool SGameplayTagPicker::CanAddNewTag() const
 {
 	const UGameplayTagsManager& Manager = UGameplayTagsManager::Get();
-	return Manager.ShouldImportTagsFromINI();
+	return !bReadOnly && Manager.ShouldImportTagsFromINI();
 }
 
 bool SGameplayTagPicker::CanAddNewSubTag(TSharedPtr<FGameplayTagNode> Node) const
diff --git a/Engine/Plugins/Editor/GameplayTagsEditor/Source/GameplayTagsEditor/Public/SGameplayTagContainerCombo.h b/Engine/Plugins/Editor/GameplayTagsEditor/Source/GameplayTagsEditor/Public/SGameplayTagContainerCombo.h
--- a/Engine/Plugins/Editor/GameplayTagsEditor/Source/GameplayTagsEditor/Public/SGameplayTagContainerCombo.h
+++ b/Engine/Plugins/Editor/GameplayTagsEditor/Source/GameplayTagsEditor/Public/SGameplayTagContainerCombo.h
@@ -90,6 +90,7 @@ private:
 		bool bMultipleValues = false;
 	};
 
+	bool IsValueEnabled() const;
 	void RefreshTagContainers();
 	TSharedRef<ITableRow> MakeTagListViewRow(TSharedPtr<FEditableItem> Item, const TSharedRef<STableViewBase>& OwnerTable);
 	void OnMenuOpenChanged(const bool bOpen);
