Fix GameplayTagQueryEntryBox didn't handle edit conditions

diff --git a/Engine/Plugins/Editor/GameplayTagsEditor/Source/GameplayTagsEditor/Private/SGameplayTagQueryEntryBox.cpp b/Engine/Plugins/Editor/GameplayTagsEditor/Source/GameplayTagsEditor/Private/SGameplayTagQueryEntryBox.cpp
--- a/Engine/Plugins/Editor/GameplayTagsEditor/Source/GameplayTagsEditor/Private/SGameplayTagQueryEntryBox.cpp
+++ b/Engine/Plugins/Editor/GameplayTagsEditor/Source/GameplayTagsEditor/Private/SGameplayTagQueryEntryBox.cpp
@@ -74,7 +74,7 @@ void SGameplayTagQueryEntryBox::Construct(const FArguments& InArgs)
 		.MaxWidth(InArgs._DescriptionMaxWidth)
 		[
 			SNew(SButton)
-			.IsEnabled(!bIsReadOnly)
+			.IsEnabled(this, &SGameplayTagQueryEntryBox::IsValueEnabled)
 			.VAlign(VAlign_Center)
 			.HAlign(HAlign_Center)
 			.ButtonStyle(FAppStyle::Get(), "SimpleButton")
@@ -96,7 +96,7 @@ void SGameplayTagQueryEntryBox::Construct(const FArguments& InArgs)
 		.VAlign(VAlign_Top)
 		[
 			SNew(SButton)
-			.IsEnabled(!bIsReadOnly)
+			.IsEnabled(this, &SGameplayTagQueryEntryBox::IsValueEnabled)
 			.ToolTipText(LOCTEXT("GameplayTagQueryEntryBox_Edit", "Edit Gameplay Tag Query."))
 			.VAlign(VAlign_Center)
 			.HAlign(HAlign_Center)
@@ -188,6 +188,16 @@ bool SGameplayTagQueryEntryBox::HasAnyValidQueries() const
 	return false;
 }
 
+bool SGameplayTagQueryEntryBox::IsValueEnabled() const
+{
+	if (PropertyHandle.IsValid())
+	{
+		return !PropertyHandle->IsEditConst();
+	}
+
+	return !bIsReadOnly;
+}
+
 EVisibility SGameplayTagQueryEntryBox::GetQueryDescVisibility() const
 {
 	return HasAnyValidQueries() == true ? EVisibility::Visible : EVisibility::Collapsed;
diff --git a/Engine/Plugins/Editor/GameplayTagsEditor/Source/GameplayTagsEditor/Public/SGameplayTagQueryEntryBox.h b/Engine/Plugins/Editor/GameplayTagsEditor/Source/GameplayTagsEditor/Public/SGameplayTagQueryEntryBox.h
--- a/Engine/Plugins/Editor/GameplayTagsEditor/Source/GameplayTagsEditor/Public/SGameplayTagQueryEntryBox.h
+++ b/Engine/Plugins/Editor/GameplayTagsEditor/Source/GameplayTagsEditor/Public/SGameplayTagQueryEntryBox.h
@@ -65,6 +65,7 @@ private:
 	FReply OnClearAllButtonClicked();
 	EVisibility GetQueryDescVisibility() const;
 	bool HasAnyValidQueries() const;
+	bool IsValueEnabled() const;
 	FText GetQueryDescText() const;
 	FText GetQueryDescTooltip() const;
 	void OnQueriesCommitted(const TArray<FGameplayTagQuery>& TagQueries);
