Fix GameplayTag Editor sometimes had a disabled OK button or a missing Clear Query button

diff --git a/Engine/Plugins/Editor/GameplayTagsEditor/Source/GameplayTagsEditor/Private/SGameplayTagQueryEntryBox.cpp b/Engine/Plugins/Editor/GameplayTagsEditor/Source/GameplayTagsEditor/Private/SGameplayTagQueryEntryBox.cpp
--- a/Engine/Plugins/Editor/GameplayTagsEditor/Source/GameplayTagsEditor/Private/SGameplayTagQueryEntryBox.cpp
+++ b/Engine/Plugins/Editor/GameplayTagsEditor/Source/GameplayTagsEditor/Private/SGameplayTagQueryEntryBox.cpp
@@ -115,7 +115,7 @@ void SGameplayTagQueryEntryBox::Construct(const FArguments& InArgs)
 		.VAlign(VAlign_Top)
 		[
 			SNew(SButton)
-			.Visibility(bIsReadOnly ? EVisibility::Collapsed : EVisibility::Visible)
+			.IsEnabled(this, &SGameplayTagQueryEntryBox::IsValueEnabled)
 			.ToolTipText(LOCTEXT("GameplayTagQueryEntryBox_Clear", "Clear Query"))
 			.VAlign(VAlign_Center)
 			.HAlign(HAlign_Center)
@@ -209,7 +209,7 @@ FReply SGameplayTagQueryEntryBox::OnEditButtonClicked()
 	Args.OnQueriesCommitted = SGameplayTagQueryWidget::FOnQueriesCommitted::CreateSP(this, &SGameplayTagQueryEntryBox::OnQueriesCommitted);
 	Args.EditableQueries = CachedQueries;
 	Args.AnchorWidget = WidgetContainer;
-	Args.bReadOnly = bIsReadOnly;
+	Args.bReadOnly = !IsValueEnabled();
 	Args.Filter = Filter;
 	
 	if (PropertyHandle.IsValid())
