Changes to SGameplayTagQueryWidget behavior.
1) Any changes made to the tag query now saved right away making it more responsive to work (previously changes would save only if you clicked OK)
2) Clicking cancel restores tag query to original state prior to any modifications ( this is to have option to get back to original query now that changes save right away)
3) Addressed positional issue when window is summoned with buttons cut of by the taskbar
4) Made the window topmost to avoid issues of clicking away and losing the query editor beneth other windows

diff --git a/Engine/Plugins/Editor/GameplayTagsEditor/Source/GameplayTagsEditor/Private/SGameplayTagQueryWidget.cpp b/Engine/Plugins/Editor/GameplayTagsEditor/Source/GameplayTagsEditor/Private/SGameplayTagQueryWidget.cpp
--- a/Engine/Plugins/Editor/GameplayTagsEditor/Source/GameplayTagsEditor/Private/SGameplayTagQueryWidget.cpp
+++ b/Engine/Plugins/Editor/GameplayTagsEditor/Source/GameplayTagsEditor/Private/SGameplayTagQueryWidget.cpp
@@ -29,9 +29,12 @@ void SGameplayTagQueryWidget::Construct(const FArguments& InArgs, const TArray<F
 	UEditableGameplayTagQuery* const EditableGameplayTagQuery = CreateEditableQuery(TagQueries[0]);
 	EditableQuery = EditableGameplayTagQuery;
 
+	OriginalTagQueries = TagQueries;
+	
 	// create details view for the editable query object
 	FDetailsViewArgs ViewArgs;
 	ViewArgs.bAllowSearch = false;
+	ViewArgs.NotifyHook = this;
 	ViewArgs.bHideSelectionTip = true;
 	ViewArgs.bShowObjectLabel = false;
 
@@ -133,6 +136,15 @@ SGameplayTagQueryWidget::~SGameplayTagQueryWidget()
 	UGameplayTagsManager::Get().OnGetCategoriesMetaFromPropertyHandle.RemoveAll(this);
 }
 
+void SGameplayTagQueryWidget::NotifyPostChange(const FPropertyChangedEvent& PropertyChangedEvent, FProperty* PropertyThatChanged)
+{
+	if (!bReadOnly)
+	{
+		FScopedTransaction Transaction( LOCTEXT("GameplayTagQueryWidget_Edit", "Edit Gameplay Tag Query") );
+		SaveToTagQuery();
+		OnQueriesCommitted.ExecuteIfBound(TagQueries);
+	}
+}
 
 void SGameplayTagQueryWidget::SaveToTagQuery()
 {
@@ -177,9 +189,13 @@ FReply SGameplayTagQueryWidget::OnOkClicked()
 	return FReply::Handled();
 }
 
-FReply SGameplayTagQueryWidget::OnCancelClicked() const
+FReply SGameplayTagQueryWidget::OnCancelClicked()
 {
+	TagQueries = OriginalTagQueries;
+	
+	OnQueriesCommitted.ExecuteIfBound(TagQueries);	
 	OnCancel.ExecuteIfBound();
+	
 	return FReply::Handled();
 }
 
@@ -194,6 +210,12 @@ TWeakPtr<SGameplayTagQueryWidget> OpenGameplayTagQueryWindow(const FGameplayTagQ
 	CloseGameplayTagQueryWindow(nullptr);
 	
 	const FVector2D WindowSize(600, 400);
+	
+	FText Title = Args.Title;
+	if (Title.IsEmpty())
+	{
+		Title = LOCTEXT("GameplayTagQueryWidget_Title", "Tag Query Editor");
+	}
 
 	// Determine the position of the window so that it will spawn near the mouse, but not go off the screen.
 	FSlateRect AnchorRect;
@@ -207,14 +229,8 @@ TWeakPtr<SGameplayTagQueryWidget> OpenGameplayTagQueryWindow(const FGameplayTagQ
 	}
 	AnchorRect = AnchorRect.ExtendBy(FMargin(20));
 	
-	const FVector2D ProposedPlacement = AnchorRect.GetTopLeft() + FVector2D(-WindowSize.X, -WindowSize.Y*0.2);
-	const FVector2D AdjustedSummonLocation = FSlateApplication::Get().CalculatePopupWindowPosition(AnchorRect, WindowSize, true, ProposedPlacement, Orient_Horizontal);
-
-	FText Title = Args.Title;
-	if (Title.IsEmpty())
-	{
-		Title = LOCTEXT("GameplayTagQueryWidget_Title", "Tag Query Editor");
-	}
+	const FVector2D ProposedPlacement = AnchorRect.GetTopLeft() + FVector2D(-WindowSize.X, -WindowSize.Y * 0.2);
+	const FVector2D AdjustedSummonLocation = FSlateApplication::Get().CalculatePopupWindowPosition(AnchorRect, WindowSize, true, ProposedPlacement, Orient_Vertical);
 
 	TSharedRef<SGameplayTagQueryWidget> QueryWidget = SNew(SGameplayTagQueryWidget, Args.EditableQueries)
 			.OnQueriesCommitted(Args.OnQueriesCommitted)
@@ -223,10 +239,10 @@ TWeakPtr<SGameplayTagQueryWidget> OpenGameplayTagQueryWindow(const FGameplayTagQ
 
 	TSharedRef<SWindow> Window = SNew(SWindow)
 		.Title(Args.Title)
-		.HasCloseButton(true)
 		.SupportsMinimize(false)
 		.SupportsMaximize(false)
 		.AutoCenter(EAutoCenter::None)
+		.IsTopmostWindow(true)
 		.ScreenPosition(AdjustedSummonLocation)
 		.ClientSize(WindowSize)
 		[
@@ -266,7 +282,7 @@ void CloseGameplayTagQueryWindow(TWeakPtr<SGameplayTagQueryWidget> QueryWidget)
 			GlobalTagQueryWidgetWindow.Pin()->RequestDestroyWindow();
 		}
 	}
-	
+
 	GlobalTagQueryWidget = nullptr;
 	GlobalTagQueryWidgetWindow = nullptr;
 }
diff --git a/Engine/Plugins/Editor/GameplayTagsEditor/Source/GameplayTagsEditor/Private/SGameplayTagQueryWidget.h b/Engine/Plugins/Editor/GameplayTagsEditor/Source/GameplayTagsEditor/Private/SGameplayTagQueryWidget.h
--- a/Engine/Plugins/Editor/GameplayTagsEditor/Source/GameplayTagsEditor/Private/SGameplayTagQueryWidget.h
+++ b/Engine/Plugins/Editor/GameplayTagsEditor/Source/GameplayTagsEditor/Private/SGameplayTagQueryWidget.h
@@ -11,7 +11,7 @@ struct FPropertyChangedEvent;
 class UEditableGameplayTagQuery;
 
 /** Widget allowing user to tag assets with gameplay tags */
-class SGameplayTagQueryWidget : public SCompoundWidget
+class SGameplayTagQueryWidget : public SCompoundWidget, public FNotifyHook
 {
 public:
 
@@ -40,6 +40,8 @@ public:
 	
 	virtual ~SGameplayTagQueryWidget() override;
 
+	virtual void NotifyPostChange( const FPropertyChangedEvent& PropertyChangedEvent, FProperty* PropertyThatChanged ) override;
+
 private:
 
 	void OnGetCategoriesMetaFromPropertyHandle(TSharedPtr<IPropertyHandle> PropertyHandle, FString& MetaString) const;
@@ -48,7 +50,7 @@ private:
 	FReply OnOkClicked();
 
 	/** Called when the user clicks the "Close Without Saving" button */
-	FReply OnCancelClicked() const;
+	FReply OnCancelClicked();
 
 	/** Controls visibility of the "Save and Close" button */
 	EVisibility GetSaveAndCloseButtonVisibility() const;
@@ -80,7 +82,10 @@ private:
 	TSharedPtr<IDetailsView> Details;
 
 	static UEditableGameplayTagQuery* CreateEditableQuery(FGameplayTagQuery& Q);
+	
 	TWeakObjectPtr<UEditableGameplayTagQuery> EditableQuery;
+	
+	TArray<FGameplayTagQuery> OriginalTagQueries;
 };
 
 struct FGameplayTagQueryWindowArgs
