Fix for [Main CIS Issue]:horde-error: Issue 341510: Compile errors in SGameplayTagQueryWidget.h and SGameplayTagQueryWidget.cpp
Adding missing .h file

Changes to SGameplayTagQueryWidget behavior.
1) Any changes made to the tag query now saved right away making it more responsive to work (previously changes would save only if you clicked OK)
2) Clicking cancel restores tag query to original state prior to any modifications ( this is to have option to get back to original query now that changes save right away)
3) Addressed positional issue when window is summoned with buttons cut of by the taskbar
4) Made the window topmost to avoid issues of clicking away and losing the query editor beneth other windows

diff --git a/Engine/Plugins/Editor/GameplayTagsEditor/Source/GameplayTagsEditor/Private/SGameplayTagQueryWidget.h b/Engine/Plugins/Editor/GameplayTagsEditor/Source/GameplayTagsEditor/Private/SGameplayTagQueryWidget.h
--- a/Engine/Plugins/Editor/GameplayTagsEditor/Source/GameplayTagsEditor/Private/SGameplayTagQueryWidget.h
+++ b/Engine/Plugins/Editor/GameplayTagsEditor/Source/GameplayTagsEditor/Private/SGameplayTagQueryWidget.h
@@ -3,6 +3,7 @@
 #pragma once
 
 #include "Widgets/SCompoundWidget.h"
+#include "Misc/NotifyHook.h"
 
 class IPropertyHandle;
 struct FGameplayTagQuery;
