Prevent accessing of invalid property nodes after undo and redo.

There's a window after undo (between undo and next property view tick) where property nodes point to the previous data.
If e.g. a property customization needs to refresh cached data after undo, they will try to access the data via old pointers.
Similar issue happens when removing item from array, the tick of IDetailCustomNodeBuilder gets called before we have updates the property nodes.

All these issues happen because when we update the child nodes, there are some parts of code that still hold hard refereces to them, and they are not
actually cleared on FPropertyNode::DestroyTree(). That in turn will cause IPropertyHandle to be in valid state, but pointing to data that has been freed.

- Detach property nodes recursively on FPropertyNode::DestroyTree(), also clearing parent pointers to that the invlidate property path cannot be followed.
- Call update view (same as tick) after undo to invalidate property nodes which do not point to correct memory anymore
- Use FEditorDelegates::PostUndoRedo instead of undo client, as in some rare cases with consequtive undo/redo the undo client callbacks get called on a stale pointer
- Postpone gameplay tag caching to next tick, since the delegates are called in reverse and SDetailsView gets OnUndoRedo gets called after out customization

diff --git a/Engine/Plugins/Editor/GameplayTagsEditor/Source/GameplayTagsEditor/Private/SGameplayTagCombo.cpp b/Engine/Plugins/Editor/GameplayTagsEditor/Source/GameplayTagsEditor/Private/SGameplayTagCombo.cpp
--- a/Engine/Plugins/Editor/GameplayTagsEditor/Source/GameplayTagsEditor/Private/SGameplayTagCombo.cpp
+++ b/Engine/Plugins/Editor/GameplayTagsEditor/Source/GameplayTagsEditor/Private/SGameplayTagCombo.cpp
@@ -31,9 +31,10 @@ SGameplayTagCombo::SGameplayTagCombo()
 
 SGameplayTagCombo::~SGameplayTagCombo()
 {
-	if (bRegisteredForUndo)
+	if (PostUndoRedoDelegateHandle.IsValid())
 	{
-		GEditor->UnregisterForUndo(this);
+		FEditorDelegates::PostUndoRedo.Remove(PostUndoRedoDelegateHandle);
+		PostUndoRedoDelegateHandle.Reset();
 	}
 }
 
@@ -48,10 +49,9 @@ void SGameplayTagCombo::Construct(const FArguments& InArgs)
 
 	if (PropertyHandle.IsValid())
 	{
+		PostUndoRedoDelegateHandle = FEditorDelegates::PostUndoRedo.AddSP(this, &SGameplayTagCombo::OnPostUndoRedo);
 		PropertyHandle->SetOnPropertyValueChanged(FSimpleDelegate::CreateSP(this, &SGameplayTagCombo::RefreshTagsFromProperty));
 		RefreshTagsFromProperty();
-		GEditor->RegisterForUndo(this);
-		bRegisteredForUndo = true;
 
 		if (Filter.IsEmpty())
 		{
@@ -93,6 +93,13 @@ void SGameplayTagCombo::Construct(const FArguments& InArgs)
 	];
 }
 
+void SGameplayTagCombo::OnPostUndoRedo()
+{
+	// This widgets OnPostUndoRedo is called before the details view has had change to handle post undo (the delegate is executed in reverse).
+	// Defer the update to next tick, so that details view has had the change to refresh the property nodes (or else we crash).
+	GEditor->GetTimerManager()->SetTimerForNextTick(FTimerDelegate::CreateSP(this, &SGameplayTagCombo::RefreshTagsFromProperty));
+}
+
 bool SGameplayTagCombo::IsValueEnabled() const
 {
 	if (PropertyHandle.IsValid())
@@ -103,22 +110,6 @@ bool SGameplayTagCombo::IsValueEnabled() const
 	return !bIsReadOnly;
 }
 
-void SGameplayTagCombo::PostUndo(bool bSuccess)
-{
-	if (bSuccess)
-	{
-		RefreshTagsFromProperty();
-	}
-}
-
-void SGameplayTagCombo::PostRedo(bool bSuccess)
-{
-	if (bSuccess)
-	{
-		RefreshTagsFromProperty();
-	}
-}
-
 FReply SGameplayTagCombo::OnEditTag() const
 {
 	FReply Reply = FReply::Handled();
@@ -366,22 +357,24 @@ bool SGameplayTagCombo::CanPaste() const
 
 void SGameplayTagCombo::RefreshTagsFromProperty()
 {
-	check(PropertyHandle.IsValid());
-
-	bHasMultipleValues = false;
-	TagsFromProperty.Reset();
-	
-	SGameplayTagPicker::EnumerateEditableTagContainersFromPropertyHandle(PropertyHandle.ToSharedRef(), [this](const FGameplayTagContainer& TagContainer)
+	if (PropertyHandle.IsValid()
+		&& PropertyHandle->IsValidHandle())
 	{
-		const FGameplayTag TagFromProperty = TagContainer.IsEmpty() ? FGameplayTag() : TagContainer.First(); 
-		if (TagsFromProperty.Num() > 0 && TagsFromProperty[0] != TagFromProperty)
+		bHasMultipleValues = false;
+		TagsFromProperty.Reset();
+		
+		SGameplayTagPicker::EnumerateEditableTagContainersFromPropertyHandle(PropertyHandle.ToSharedRef(), [this](const FGameplayTagContainer& TagContainer)
 		{
-			bHasMultipleValues = true;
-		}
-		TagsFromProperty.Add(TagFromProperty);
+			const FGameplayTag TagFromProperty = TagContainer.IsEmpty() ? FGameplayTag() : TagContainer.First(); 
+			if (TagsFromProperty.Num() > 0 && TagsFromProperty[0] != TagFromProperty)
+			{
+				bHasMultipleValues = true;
+			}
+			TagsFromProperty.Add(TagFromProperty);
 
-		return true;
-	});
+			return true;
+		});
+	}
 }
 
 #undef LOCTEXT_NAMESPACE
diff --git a/Engine/Plugins/Editor/GameplayTagsEditor/Source/GameplayTagsEditor/Private/SGameplayTagContainerCombo.cpp b/Engine/Plugins/Editor/GameplayTagsEditor/Source/GameplayTagsEditor/Private/SGameplayTagContainerCombo.cpp
--- a/Engine/Plugins/Editor/GameplayTagsEditor/Source/GameplayTagsEditor/Private/SGameplayTagContainerCombo.cpp
+++ b/Engine/Plugins/Editor/GameplayTagsEditor/Source/GameplayTagsEditor/Private/SGameplayTagContainerCombo.cpp
@@ -37,9 +37,10 @@ SGameplayTagContainerCombo::SGameplayTagContainerCombo()
 
 SGameplayTagContainerCombo::~SGameplayTagContainerCombo()
 {
-	if (bRegisteredForUndo)
+	if (PostUndoRedoDelegateHandle.IsValid())
 	{
-		GEditor->UnregisterForUndo(this);
+		FEditorDelegates::PostUndoRedo.Remove(PostUndoRedoDelegateHandle);
+		PostUndoRedoDelegateHandle.Reset();
 	}
 }
 
@@ -53,10 +54,9 @@ void SGameplayTagContainerCombo::Construct(const FArguments& InArgs)
 
 	if (PropertyHandle.IsValid())
 	{
+		PostUndoRedoDelegateHandle = FEditorDelegates::PostUndoRedo.AddSP(this, &SGameplayTagContainerCombo::OnPostUndoRedo);
 		PropertyHandle->SetOnPropertyValueChanged(FSimpleDelegate::CreateSP(this, &SGameplayTagContainerCombo::RefreshTagContainers));
 		RefreshTagContainers();
-		GEditor->RegisterForUndo(this);
-		bRegisteredForUndo = true;
 
 		if (Filter.IsEmpty())
 		{
@@ -158,6 +158,13 @@ void SGameplayTagContainerCombo::Construct(const FArguments& InArgs)
 	];
 }
 
+void SGameplayTagContainerCombo::OnPostUndoRedo()
+{
+	// This widgets OnPostUndoRedo is called before the details view has had change to handle post undo (the delegate is executed in reverse).
+	// Defer the update to next tick, so that details view has had the change to refresh the property nodes (or else we crash).
+	GEditor->GetTimerManager()->SetTimerForNextTick( FTimerDelegate::CreateSP(this, &SGameplayTagContainerCombo::RefreshTagContainers));
+}
+
 bool SGameplayTagContainerCombo::IsValueEnabled() const
 {
 	if (PropertyHandle.IsValid())
@@ -511,22 +518,6 @@ FReply SGameplayTagContainerCombo::OnClearTagClicked(const FGameplayTag TagToCle
 	return FReply::Handled();
 }
 
-void SGameplayTagContainerCombo::PostUndo(bool bSuccess)
-{
-	if (bSuccess)
-	{
-		RefreshTagContainers();
-	}
-}
-
-void SGameplayTagContainerCombo::PostRedo(bool bSuccess)
-{
-	if (bSuccess)
-	{
-		RefreshTagContainers();
-	}
-}
-
 void SGameplayTagContainerCombo::RefreshTagContainers()
 {
 	CachedTagContainers.Reset();
@@ -534,30 +525,33 @@ void SGameplayTagContainerCombo::RefreshTagContainers()
 
 	if (PropertyHandle.IsValid())
 	{
-		// From property
-		SGameplayTagPicker::EnumerateEditableTagContainersFromPropertyHandle(PropertyHandle.ToSharedRef(), [this](const FGameplayTagContainer& InTagContainer)
+		if (PropertyHandle->IsValidHandle())
 		{
-			CachedTagContainers.Add(InTagContainer);
-
-			for (auto It = InTagContainer.CreateConstIterator(); It; ++It)
+			// From property
+			SGameplayTagPicker::EnumerateEditableTagContainersFromPropertyHandle(PropertyHandle.ToSharedRef(), [this](const FGameplayTagContainer& InTagContainer)
 			{
-				const FGameplayTag Tag = *It;
-				const int32 ExistingItemIndex = TagsToEdit.IndexOfByPredicate([Tag](const TSharedPtr<FEditableItem>& Item)
-				{
-					return Item.IsValid() && Item->Tag == Tag;
-				});
-				if (ExistingItemIndex != INDEX_NONE)
-				{
-					TagsToEdit[ExistingItemIndex]->Count++;
-				}
-				else
+				CachedTagContainers.Add(InTagContainer);
+
+				for (auto It = InTagContainer.CreateConstIterator(); It; ++It)
 				{
-					TagsToEdit.Add(MakeShared<FEditableItem>(Tag));
+					const FGameplayTag Tag = *It;
+					const int32 ExistingItemIndex = TagsToEdit.IndexOfByPredicate([Tag](const TSharedPtr<FEditableItem>& Item)
+					{
+						return Item.IsValid() && Item->Tag == Tag;
+					});
+					if (ExistingItemIndex != INDEX_NONE)
+					{
+						TagsToEdit[ExistingItemIndex]->Count++;
+					}
+					else
+					{
+						TagsToEdit.Add(MakeShared<FEditableItem>(Tag));
+					}
 				}
-			}
-			
-			return true;
-		});
+				
+				return true;
+			});
+		}
 	}
 	else
 	{
diff --git a/Engine/Plugins/Editor/GameplayTagsEditor/Source/GameplayTagsEditor/Private/SGameplayTagPicker.cpp b/Engine/Plugins/Editor/GameplayTagsEditor/Source/GameplayTagsEditor/Private/SGameplayTagPicker.cpp
--- a/Engine/Plugins/Editor/GameplayTagsEditor/Source/GameplayTagsEditor/Private/SGameplayTagPicker.cpp
+++ b/Engine/Plugins/Editor/GameplayTagsEditor/Source/GameplayTagsEditor/Private/SGameplayTagPicker.cpp
@@ -75,6 +75,15 @@ bool SGameplayTagPicker::GetEditableTagContainersFromPropertyHandle(const TShare
 	});
 }
 
+SGameplayTagPicker::~SGameplayTagPicker()
+{
+	if (PostUndoRedoDelegateHandle.IsValid())
+	{
+		FEditorDelegates::PostUndoRedo.Remove(PostUndoRedoDelegateHandle);
+		PostUndoRedoDelegateHandle.Reset();
+	}
+}
+
 void SGameplayTagPicker::Construct(const FArguments& InArgs)
 {
 	TagContainers = InArgs._TagContainers;
@@ -101,6 +110,8 @@ void SGameplayTagPicker::Construct(const FArguments& InArgs)
 
 	bRestrictedTags = InArgs._RestrictedTags;
 
+	PostUndoRedoDelegateHandle = FEditorDelegates::PostUndoRedo.AddSP(this, &SGameplayTagPicker::OnPostUndoRedo);
+	
 	UGameplayTagsManager::OnEditorRefreshGameplayTagTree.AddSP(this, &SGameplayTagPicker::RefreshOnNextTick);
 	UGameplayTagsManager& Manager = UGameplayTagsManager::Get();
 
@@ -148,7 +159,7 @@ void SGameplayTagPicker::Construct(const FArguments& InArgs)
 	SettingsCombo->SetOnGetMenuContent(FOnGetContent::CreateSP(this, &SGameplayTagPicker::MakeSettingsMenu, SettingsCombo));
 
 
-	TWeakPtr<SGameplayTagPicker> WeakSelf = StaticCastWeakPtr<SGameplayTagPicker>(AsWeak());
+	TWeakPtr<SGameplayTagPicker> WeakSelf = SharedThis(this);
 	
 	TSharedRef<SWidget> Picker = 
 		SNew(SBorder)
@@ -1775,14 +1786,17 @@ void SGameplayTagPicker::SetTagContainers(TConstArrayView<FGameplayTagContainer>
 	TagContainers = InTagContainers;
 }
 
-void SGameplayTagPicker::PostUndo(bool bSuccess)
-{
-	OnRefreshTagContainers.ExecuteIfBound(*this);
-}
-
-void SGameplayTagPicker::PostRedo(bool bSuccess)
+void SGameplayTagPicker::OnPostUndoRedo()
 {
-	OnRefreshTagContainers.ExecuteIfBound(*this);
+	// This widgets OnPostUndoRedo is called before the details view has had change to handle post undo (the delegate is executed in reverse).
+	// Defer the update to next tick, so that details view has had the change to refresh the property nodes (or else we may crash).
+	GEditor->GetTimerManager()->SetTimerForNextTick([WeakSelf = SharedThis(this).ToWeakPtr()]()
+	{
+		if (SGameplayTagPicker* Self = WeakSelf.Pin().Get())
+		{
+			Self->OnRefreshTagContainers.ExecuteIfBound(*Self);
+		}
+	});
 }
 
 void SGameplayTagPicker::VerifyAssetTagValidity()
diff --git a/Engine/Plugins/Editor/GameplayTagsEditor/Source/GameplayTagsEditor/Private/SGameplayTagQueryEntryBox.cpp b/Engine/Plugins/Editor/GameplayTagsEditor/Source/GameplayTagsEditor/Private/SGameplayTagQueryEntryBox.cpp
--- a/Engine/Plugins/Editor/GameplayTagsEditor/Source/GameplayTagsEditor/Private/SGameplayTagQueryEntryBox.cpp
+++ b/Engine/Plugins/Editor/GameplayTagsEditor/Source/GameplayTagsEditor/Private/SGameplayTagQueryEntryBox.cpp
@@ -33,9 +33,10 @@ SGameplayTagQueryEntryBox::SGameplayTagQueryEntryBox()
 
 SGameplayTagQueryEntryBox::~SGameplayTagQueryEntryBox()
 {
-	if (bRegisteredForUndo)
+	if (PostUndoRedoDelegateHandle.IsValid())
 	{
-		GEditor->UnregisterForUndo(this);
+		FEditorDelegates::PostUndoRedo.Remove(PostUndoRedoDelegateHandle);
+		PostUndoRedoDelegateHandle.Reset();
 	}
 }
 
@@ -47,9 +48,8 @@ void SGameplayTagQueryEntryBox::Construct(const FArguments& InArgs)
 
 	if (PropertyHandle.IsValid())
 	{
+		PostUndoRedoDelegateHandle = FEditorDelegates::PostUndoRedo.AddSP(this, &SGameplayTagQueryEntryBox::OnPostUndoRedo);
 		PropertyHandle->SetOnPropertyValueChanged(FSimpleDelegate::CreateSP(this, &SGameplayTagQueryEntryBox::CacheQueryList));
-		GEditor->RegisterForUndo(this);
-		bRegisteredForUndo = true;
 		bIsReadOnly = PropertyHandle->IsEditConst();
 
 		if (Filter.IsEmpty())
@@ -130,16 +130,6 @@ void SGameplayTagQueryEntryBox::Construct(const FArguments& InArgs)
 	];
 }
 
-void SGameplayTagQueryEntryBox::PostUndo(bool bSuccess)
-{
-	CacheQueryList();
-}
-
-void SGameplayTagQueryEntryBox::PostRedo(bool bSuccess)
-{
-	CacheQueryList();
-}
-
 FText SGameplayTagQueryEntryBox::GetQueryDescText() const
 {
 	return QueryDescription;
@@ -252,22 +242,32 @@ FReply SGameplayTagQueryEntryBox::OnEditButtonClicked()
 	return FReply::Handled();
 }
 
+void SGameplayTagQueryEntryBox::OnPostUndoRedo()
+{
+	// This widgets OnPostUndoRedo is called before the details view has had change to handle post undo (the delegate is executed in reverse).
+	// Defer the update to next tick, so that details view has had the change to refresh the property nodes (or else we crash).
+	GEditor->GetTimerManager()->SetTimerForNextTick(FTimerDelegate::CreateSP(this, &SGameplayTagQueryEntryBox::CacheQueryList));
+}
+
 void SGameplayTagQueryEntryBox::CacheQueryList()
 {
 	CachedQueries.Empty();
 
 	if (PropertyHandle.IsValid())
 	{
-		// Cache queries from the property handle. Add empty queries even if the instance data is null so that the indices match with the property handle.
-		TArray<void*> RawStructData;
-		PropertyHandle->AccessRawData(RawStructData);
-		
-		for (int32 Idx = 0; Idx < RawStructData.Num(); ++Idx)
+		if (PropertyHandle->IsValidHandle())
 		{
-			FGameplayTagQuery& Query = CachedQueries.AddDefaulted_GetRef();
-			if (RawStructData[Idx])
+			// Cache queries from the property handle. Add empty queries even if the instance data is null so that the indices match with the property handle.
+			TArray<void*> RawStructData;
+			PropertyHandle->AccessRawData(RawStructData);
+			
+			for (int32 Idx = 0; Idx < RawStructData.Num(); ++Idx)
 			{
-				Query = *(FGameplayTagQuery*)RawStructData[Idx]; 
+				FGameplayTagQuery& Query = CachedQueries.AddDefaulted_GetRef();
+				if (RawStructData[Idx])
+				{
+					Query = *(FGameplayTagQuery*)RawStructData[Idx]; 
+				}
 			}
 		}
 	}
diff --git a/Engine/Plugins/Editor/GameplayTagsEditor/Source/GameplayTagsEditor/Public/SGameplayTagCombo.h b/Engine/Plugins/Editor/GameplayTagsEditor/Source/GameplayTagsEditor/Public/SGameplayTagCombo.h
--- a/Engine/Plugins/Editor/GameplayTagsEditor/Source/GameplayTagsEditor/Public/SGameplayTagCombo.h
+++ b/Engine/Plugins/Editor/GameplayTagsEditor/Source/GameplayTagsEditor/Public/SGameplayTagCombo.h
@@ -7,7 +7,6 @@
 #include "Widgets/SCompoundWidget.h"
 #include "Widgets/DeclarativeSyntaxSupport.h"
 #include "GameplayTagContainer.h"
-#include "EditorUndoClient.h"
 #include "SGameplayTagChip.h"
 
 class IPropertyHandle;
@@ -19,7 +18,7 @@ class SGameplayTagPicker;
 /**
  * Widget for editing a Gameplay Tag.
  */
-class SGameplayTagCombo : public SCompoundWidget, public FEditorUndoClient
+class SGameplayTagCombo : public SCompoundWidget
 {
 	SLATE_DECLARE_WIDGET(SGameplayTagCombo, SCompoundWidget)
 
@@ -66,14 +65,9 @@ public:
 
 	GAMEPLAYTAGSEDITOR_API void Construct(const FArguments& InArgs);
 
-protected:
-	//~ Begin FEditorUndoClient Interface
-	virtual void PostUndo(bool bSuccess) override;
-	virtual void PostRedo(bool bSuccess) override;
-	//~ End FEditorUndoClient Interface
-
 private:
 
+	void OnPostUndoRedo();
 	bool ShowClearButton() const;
 	FText GetText() const;
 	bool IsValueEnabled() const;
@@ -98,10 +92,10 @@ private:
 	bool bIsReadOnly = false;
 	FString Filter;
 	FString SettingsName;
-	bool bRegisteredForUndo = false;
 	FOnTagChanged OnTagChanged;
 	TSharedPtr<IPropertyHandle> PropertyHandle;
 	TSharedPtr<SMenuAnchor> MenuAnchor;
 	TSharedPtr<SComboButton> ComboButton;
 	TSharedPtr<SGameplayTagPicker> TagPicker;
+	FDelegateHandle PostUndoRedoDelegateHandle;
 };
diff --git a/Engine/Plugins/Editor/GameplayTagsEditor/Source/GameplayTagsEditor/Public/SGameplayTagContainerCombo.h b/Engine/Plugins/Editor/GameplayTagsEditor/Source/GameplayTagsEditor/Public/SGameplayTagContainerCombo.h
--- a/Engine/Plugins/Editor/GameplayTagsEditor/Source/GameplayTagsEditor/Public/SGameplayTagContainerCombo.h
+++ b/Engine/Plugins/Editor/GameplayTagsEditor/Source/GameplayTagsEditor/Public/SGameplayTagContainerCombo.h
@@ -8,7 +8,6 @@
 #include "Widgets/Views/SListView.h"
 #include "Widgets/DeclarativeSyntaxSupport.h"
 #include "GameplayTagContainer.h"
-#include "EditorUndoClient.h"
 #include "SGameplayTagChip.h"
 
 class IPropertyHandle;
@@ -21,7 +20,7 @@ class SComboButton;
 /**
  * Widget for editing a Gameplay Tag Container.
  */
-class SGameplayTagContainerCombo : public SCompoundWidget, public FEditorUndoClient
+class SGameplayTagContainerCombo : public SCompoundWidget
 {
 	SLATE_DECLARE_WIDGET(SGameplayTagContainerCombo, SCompoundWidget)
 	
@@ -68,12 +67,6 @@ public:
 
 	GAMEPLAYTAGSEDITOR_API void Construct(const FArguments& InArgs);
 
-protected:
-	//~ Begin FEditorUndoClient Interface
-	virtual void PostUndo(bool bSuccess) override;
-	virtual void PostRedo(bool bSuccess) override;
-	//~ End FEditorUndoClient Interface
-
 private:
 
 	struct FEditableItem
@@ -90,6 +83,7 @@ private:
 		bool bMultipleValues = false;
 	};
 
+	void OnPostUndoRedo();
 	bool IsValueEnabled() const;
 	void RefreshTagContainers();
 	TSharedRef<ITableRow> MakeTagListViewRow(TSharedPtr<FEditableItem> Item, const TSharedRef<STableViewBase>& OwnerTable);
@@ -115,11 +109,11 @@ private:
 	FString Filter;
 	FString SettingsName;
 	bool bIsReadOnly = false;
-	bool bRegisteredForUndo = false;
 	FOnTagContainerChanged OnTagContainerChanged;
 	TSharedPtr<IPropertyHandle> PropertyHandle;
 	TSharedPtr<SMenuAnchor> MenuAnchor;
 	TSharedPtr<SComboButton> ComboButton;
 	TSharedPtr<SGameplayTagPicker> TagPicker;
 	FGameplayTag TagToHilight;
+	FDelegateHandle PostUndoRedoDelegateHandle;
 };
diff --git a/Engine/Plugins/Editor/GameplayTagsEditor/Source/GameplayTagsEditor/Public/SGameplayTagPicker.h b/Engine/Plugins/Editor/GameplayTagsEditor/Source/GameplayTagsEditor/Public/SGameplayTagPicker.h
--- a/Engine/Plugins/Editor/GameplayTagsEditor/Source/GameplayTagsEditor/Public/SGameplayTagPicker.h
+++ b/Engine/Plugins/Editor/GameplayTagsEditor/Source/GameplayTagsEditor/Public/SGameplayTagPicker.h
@@ -13,7 +13,6 @@
 #include "Widgets/Views/STableRow.h"
 #include "Widgets/Views/STreeView.h"
 #include "GameplayTagsManager.h"
-#include "EditorUndoClient.h"
 
 class IPropertyHandle;
 class SComboButton;
@@ -28,7 +27,7 @@ enum class EGameplayTagPickerMode : uint8
 };
 
 /** Widget allowing user to tag assets with gameplay tags */
-class GAMEPLAYTAGSEDITOR_API SGameplayTagPicker : public SCompoundWidget, public FSelfRegisteringEditorUndoClient
+class GAMEPLAYTAGSEDITOR_API SGameplayTagPicker : public SCompoundWidget
 {
 public:
 
@@ -114,6 +113,8 @@ public:
 	 */
 	static bool GetEditableTagContainersFromPropertyHandle(const TSharedRef<IPropertyHandle>& PropHandle, TArray<FGameplayTagContainer>& OutEditableContainers);
 
+	virtual ~SGameplayTagPicker() override;
+
 	/** Construct the actual widget */
 	void Construct(const FArguments& InArgs);
 	
@@ -152,10 +153,7 @@ private:
 		Duplicate,
 	};
 
-	// FSelfRegisteringEditorUndoClient
-	virtual void PostUndo(bool bSuccess) override;
-	virtual void PostRedo(bool bSuccess) override;
-	// ~FSelfRegisteringEditorUndoClient
+	void OnPostUndoRedo();
 
 	/** Verify the tags are all valid and if not prompt the user. */
 	void VerifyAssetTagValidity();
@@ -452,6 +450,8 @@ private:
 
 	TSharedPtr<SAddNewGameplayTagWidget> AddNewTagWidget;
 	bool bNewTagWidgetVisible = false;
+
+	FDelegateHandle PostUndoRedoDelegateHandle;
 };
 
 
diff --git a/Engine/Plugins/Editor/GameplayTagsEditor/Source/GameplayTagsEditor/Public/SGameplayTagQueryEntryBox.h b/Engine/Plugins/Editor/GameplayTagsEditor/Source/GameplayTagsEditor/Public/SGameplayTagQueryEntryBox.h
--- a/Engine/Plugins/Editor/GameplayTagsEditor/Source/GameplayTagsEditor/Public/SGameplayTagQueryEntryBox.h
+++ b/Engine/Plugins/Editor/GameplayTagsEditor/Source/GameplayTagsEditor/Public/SGameplayTagQueryEntryBox.h
@@ -6,7 +6,6 @@
 #include "Widgets/SCompoundWidget.h"
 #include "Widgets/DeclarativeSyntaxSupport.h"
 #include "GameplayTagContainer.h"
-#include "EditorUndoClient.h"
 #include "SGameplayTagQueryWidget.h"
 
 class IPropertyHandle;
@@ -15,7 +14,7 @@ class SHorizontalBox;
 /**
  * Widget for editing a Gameplay Tag Query.
  */
-class SGameplayTagQueryEntryBox : public SCompoundWidget, public FEditorUndoClient
+class SGameplayTagQueryEntryBox : public SCompoundWidget
 {
 	SLATE_DECLARE_WIDGET(SGameplayTagQueryEntryBox, SCompoundWidget)
 	
@@ -52,14 +51,9 @@ public:
 
 	GAMEPLAYTAGSEDITOR_API void Construct(const FArguments& InArgs);
 
-protected:
-	//~ Begin FEditorUndoClient Interface
-	virtual void PostUndo(bool bSuccess) override;
-	virtual void PostRedo(bool bSuccess) override;
-	//~ End FEditorUndoClient Interface
-
 private:
 
+	void OnPostUndoRedo();
 	void CacheQueryList();
 	FReply OnEditButtonClicked();
 	FReply OnClearAllButtonClicked();
@@ -74,7 +68,6 @@ private:
 
 	bool bIsReadOnly = false;
 	FString Filter;
-	bool bRegisteredForUndo = false;
 	FOnTagQueryChanged OnTagQueryChanged;
 	TSharedPtr<IPropertyHandle> PropertyHandle;
 	TArray<FGameplayTagQuery> CachedQueries;
@@ -85,4 +78,6 @@ private:
 	
 	FText QueryDescription;
 	FText QueryDescriptionTooltip;
+	
+	FDelegateHandle PostUndoRedoDelegateHandle;
 };
diff --git a/Engine/Plugins/Experimental/StructUtils/Source/StructUtilsEditor/Private/Customizations/InstancedStructDetails.cpp b/Engine/Plugins/Experimental/StructUtils/Source/StructUtilsEditor/Private/Customizations/InstancedStructDetails.cpp
--- a/Engine/Plugins/Experimental/StructUtils/Source/StructUtilsEditor/Private/Customizations/InstancedStructDetails.cpp
+++ b/Engine/Plugins/Experimental/StructUtils/Source/StructUtilsEditor/Private/Customizations/InstancedStructDetails.cpp
@@ -11,7 +11,6 @@
 #include "ScopedTransaction.h"
 #include "Modules/ModuleManager.h"
 #include "StructViewerModule.h"
-#include "StructViewerFilter.h"
 #include "Styling/SlateIconFinder.h"
 #include "Engine/UserDefinedStruct.h"
 #include "InstancedStruct.h"
@@ -351,21 +350,9 @@ FName FInstancedStructDataDetails::GetName() const
 	return Name;
 }
 
-void FInstancedStructDataDetails::PostUndo(bool bSuccess)
-{
-	// Undo; force a sync next Tick
-	LastSyncEditableInstanceFromSourceSeconds = 0.0;
-}
-
-void FInstancedStructDataDetails::PostRedo(bool bSuccess)
-{
-	// Redo; force a sync next Tick
-	LastSyncEditableInstanceFromSourceSeconds = 0.0;
-}
 
 ////////////////////////////////////
 
-
 TSharedRef<IPropertyTypeCustomization> FInstancedStructDetails::MakeInstance()
 {
 	return MakeShared<FInstancedStructDetails>();
diff --git a/Engine/Plugins/Experimental/StructUtils/Source/StructUtilsEditor/Public/InstancedStructDetails.h b/Engine/Plugins/Experimental/StructUtils/Source/StructUtilsEditor/Public/InstancedStructDetails.h
--- a/Engine/Plugins/Experimental/StructUtils/Source/StructUtilsEditor/Public/InstancedStructDetails.h
+++ b/Engine/Plugins/Experimental/StructUtils/Source/StructUtilsEditor/Public/InstancedStructDetails.h
@@ -5,6 +5,7 @@
 #include "IPropertyTypeCustomization.h"
 #include "IDetailCustomNodeBuilder.h"
 #include "EditorUndoClient.h"
+#include "StructViewerFilter.h"
 
 class IPropertyHandle;
 class IDetailPropertyRow;
@@ -59,11 +60,11 @@ private:
  * Can be used in a implementation of a IPropertyTypeCustomization CustomizeChildren() to display editable FInstancedStruct contents.
  * OnChildRowAdded() is called right after each property is added, which allows the property row to be customizable.
  */
-class STRUCTUTILSEDITOR_API FInstancedStructDataDetails : public IDetailCustomNodeBuilder, public FSelfRegisteringEditorUndoClient, public TSharedFromThis<FInstancedStructDataDetails>
+class STRUCTUTILSEDITOR_API FInstancedStructDataDetails : public IDetailCustomNodeBuilder, public TSharedFromThis<FInstancedStructDataDetails>
 {
 public:
 	FInstancedStructDataDetails(TSharedPtr<IPropertyHandle> InStructProperty);
-	~FInstancedStructDataDetails();
+	virtual ~FInstancedStructDataDetails() override;
 
 	//~ Begin IDetailCustomNodeBuilder interface
 	virtual void SetOnRebuildChildren(FSimpleDelegate InOnRegenerateChildren) override;
@@ -75,10 +76,6 @@ public:
 	virtual FName GetName() const override;
 	//~ End IDetailCustomNodeBuilder interface
 
-	/** FEditorUndoClient interface */
-	virtual void PostUndo(bool bSuccess) override;
-	virtual void PostRedo(bool bSuccess) override;
-
 	// Called when a child is added, override to customize a child row.
 	virtual void OnChildRowAdded(IDetailPropertyRow& ChildRow) {}
 
@@ -105,9 +102,6 @@ private:
 	/** Delegate that can be used to refresh the child rows of the current struct (eg, when changing struct type) */
 	FSimpleDelegate OnRegenerateChildren;
 
-	/** The last time that SyncEditableInstanceFromSource was called, in FPlatformTime::Seconds() */
-	double LastSyncEditableInstanceFromSourceSeconds = 0.0;
-
 	/** True if we're currently handling a StructValuePostChange */
 	bool bIsHandlingStructValuePostChange = false;
 	
diff --git a/Engine/Source/Editor/PropertyEditor/Private/PropertyHandleImpl.cpp b/Engine/Source/Editor/PropertyEditor/Private/PropertyHandleImpl.cpp
--- a/Engine/Source/Editor/PropertyEditor/Private/PropertyHandleImpl.cpp
+++ b/Engine/Source/Editor/PropertyEditor/Private/PropertyHandleImpl.cpp
@@ -56,6 +56,11 @@ FPropertyValueImpl::FPropertyValueImpl( TSharedPtr<FPropertyNode> InPropertyNode
 
 void FPropertyValueImpl::EnumerateObjectsToModify( FPropertyNode* InPropertyNode, const EnumerateObjectsToModifyFuncRef& InObjectsToModifyCallback ) const
 {
+	if (!InPropertyNode)
+	{
+		return;
+	}
+	
 	// Find the parent object node which contains offset addresses for reading a property value on an object
 	FComplexPropertyNode* ComplexNode = InPropertyNode->FindComplexParent();
 	if (ComplexNode)
@@ -834,17 +839,18 @@ FPropertyAccess::Result FPropertyValueImpl::SetValueAsString( const FString& InV
 
 		// If more than one object is selected, an empty field indicates their values for this property differ.
 		// Don't send it to the objects value in this case (if we did, they would all get set to None which isn't good).
-		FComplexPropertyNode* ParentNode = PropertyNodePin->FindComplexParent();
-
-		FString PreviousValue;
-		GetValueAsString( PreviousValue );
+		if (FComplexPropertyNode* ParentNode = PropertyNodePin->FindComplexParent())
+		{
+			FString PreviousValue;
+			GetValueAsString( PreviousValue );
 
-		const bool bDidValueChange = Value.Len() && (FCString::Strcmp(*PreviousValue, *Value) != 0);
-		const bool bComingOutOfInteractiveChange = bInteractiveChangeInProgress && ( ( Flags & EPropertyValueSetFlags::InteractiveChange ) != EPropertyValueSetFlags::InteractiveChange );
+			const bool bDidValueChange = Value.Len() && (FCString::Strcmp(*PreviousValue, *Value) != 0);
+			const bool bComingOutOfInteractiveChange = bInteractiveChangeInProgress && ( ( Flags & EPropertyValueSetFlags::InteractiveChange ) != EPropertyValueSetFlags::InteractiveChange );
 
-		if ( ParentNode && ( ParentNode->GetInstancesNum() == 1 || bComingOutOfInteractiveChange || bDidValueChange ) )
-		{
-			ImportText( Value, PropertyNodePin.Get(), Flags );
+			if ( ParentNode->GetInstancesNum() == 1 || bComingOutOfInteractiveChange || bDidValueChange )
+			{
+				ImportText( Value, PropertyNodePin.Get(), Flags );
+			}
 		}
 
 		Result = FPropertyAccess::Success;
diff --git a/Engine/Source/Editor/PropertyEditor/Private/PropertyNode.cpp b/Engine/Source/Editor/PropertyEditor/Private/PropertyNode.cpp
--- a/Engine/Source/Editor/PropertyEditor/Private/PropertyNode.cpp
+++ b/Engine/Source/Editor/PropertyEditor/Private/PropertyNode.cpp
@@ -90,8 +90,8 @@ FPropertyNode::~FPropertyNode()
 void FPropertyNode::InitNode(const FPropertyNodeInitParams& InitParams)
 {
 	//Dismantle the previous tree
-	DestroyTree();
-
+	DestroyTree(/*bInDestroySelf*/false);
+	
 	//tree hierarchy
 	check(InitParams.ParentNode.Get() != this);
 	ParentNodeWeakPtr = InitParams.ParentNode;
@@ -255,14 +255,15 @@ void FPropertyNode::InitNode(const FPropertyNodeInitParams& InitParams)
 		SetNodeFlags( EPropertyNodeFlags::RequiresValidation, bRequiresValidation );
 	}
 
+	// Build property path before building children so that we have parent property path for debugging.
+	PropertyPath = FPropertyNode::CreatePropertyPath(this->AsShared())->ToString();
+	
 	InitExpansionFlags();
 
 	if (InitParams.bAllowChildren)
 	{
 		RebuildChildren();
 	}
-
-	PropertyPath = FPropertyNode::CreatePropertyPath(this->AsShared())->ToString();
 }
 
 /**
@@ -2667,7 +2668,25 @@ bool FPropertyNode::IsChildOfFavorite (void) const
  */
 void FPropertyNode::DestroyTree(const bool bInDestroySelf)
 {
-	ChildNodes.Empty();
+	if (bInDestroySelf)
+	{
+		// Detach from parent, so that instances of this node will not be able to follow the parent chain.
+		// Some code, like property rows, hold hard references to the nodes, so this method does not actually
+		// destroy the child nodes, it happens much later when we refresh the property nodes.
+		// Breaking the chain prevents the "destroyed" nodes accessing the parent (complex) nodes,
+		// which may already point to completely different data (and have new child nodes).
+		ParentNodeWeakPtr = nullptr;
+	}
+	
+	if (ChildNodes.Num() > 0)
+	{
+		// Detach and drop children.
+		for (TSharedPtr<FPropertyNode>& ChildNode : ChildNodes)
+		{
+			ChildNode->DestroyTree();
+		}
+		ChildNodes.Empty();
+	}
 }
 
 /**
diff --git a/Engine/Source/Editor/PropertyEditor/Private/SDetailsView.cpp b/Engine/Source/Editor/PropertyEditor/Private/SDetailsView.cpp
--- a/Engine/Source/Editor/PropertyEditor/Private/SDetailsView.cpp
+++ b/Engine/Source/Editor/PropertyEditor/Private/SDetailsView.cpp
@@ -1030,6 +1030,9 @@ void SDetailsView::InvalidateCachedState()
 	{
 		RootNode->InvalidateCachedState();
 	}
+
+	// Update property view immediately after undo or redo so to validate that the data layout matches the current data. 
+	UpdatePropertyView(0.0f);
 }
 
 bool SDetailsView::IsConnected() const
diff --git a/Engine/Source/Editor/PropertyEditor/Private/SDetailsViewBase.cpp b/Engine/Source/Editor/PropertyEditor/Private/SDetailsViewBase.cpp
--- a/Engine/Source/Editor/PropertyEditor/Private/SDetailsViewBase.cpp
+++ b/Engine/Source/Editor/PropertyEditor/Private/SDetailsViewBase.cpp
@@ -1083,10 +1083,14 @@ void SDetailsViewBase::HandlePendingCleanup()
 }
 
 /** Ticks the property view.  This function performs a data consistency check */
-void SDetailsViewBase::Tick( const FGeometry& AllottedGeometry, const double InCurrentTime, const float InDeltaTime )
+void SDetailsViewBase::Tick( const FGeometry& /*AllottedGeometry*/, const double InCurrentTime, const float InDeltaTime )
 {
 	TRACE_CPUPROFILER_EVENT_SCOPE(SDetailsViewBase::Tick);
+	UpdatePropertyView(InDeltaTime);
+}
 
+void SDetailsViewBase::UpdatePropertyView(const float InDeltaTime)
+{
 	HandlePendingCleanup();
 
 	FDetailsViewConfig* ViewConfig = GetMutableViewConfig();
diff --git a/Engine/Source/Editor/PropertyEditor/Private/SDetailsViewBase.h b/Engine/Source/Editor/PropertyEditor/Private/SDetailsViewBase.h
--- a/Engine/Source/Editor/PropertyEditor/Private/SDetailsViewBase.h
+++ b/Engine/Source/Editor/PropertyEditor/Private/SDetailsViewBase.h
@@ -187,6 +187,9 @@ public:
 	// End of SWidget interface
 
 protected:
+
+	void UpdatePropertyView(const float InDeltaTime);
+
 	/**
 	 * Called when a color property is changed from a color picker
 	 */
