Prevent accessing of invalid property nodes after undo and redo.

There's a window after undo (between undo and next property view tick) where property nodes point to the previous data.
If e.g. a property customization needs to refresh cached data after undo, they will try to access the data via old pointers.
Similar issue happens when removing item from array, the tick of IDetailCustomNodeBuilder gets called before we have updates the property nodes.

All these issues happen because when we update the child nodes, there are some parts of code that still hold hard refereces to them, and they are not
actually cleared on FPropertyNode::DestroyTree(). That in turn will cause IPropertyHandle to be in valid state, but pointing to data that has been freed.

- Added bIsDestroyed on property nodes to be able distinguish property nodes which are still alive but destoyed
- Prevent calling Tick() on property tree nodes, whose property node is destroyed
- Changed Gameplay Tag, Tag Container and Query to poll for changes, because undo callbacks are triggered at times when the property handle (and underlying node) is pointing to old data
- The above gameplay tag & co change will also react to property changes from outside the details panel

diff --git a/Engine/Plugins/Editor/GameplayTagsEditor/Source/GameplayTagsEditor/Private/SGameplayTagCombo.cpp b/Engine/Plugins/Editor/GameplayTagsEditor/Source/GameplayTagsEditor/Private/SGameplayTagCombo.cpp
--- a/Engine/Plugins/Editor/GameplayTagsEditor/Source/GameplayTagsEditor/Private/SGameplayTagCombo.cpp
+++ b/Engine/Plugins/Editor/GameplayTagsEditor/Source/GameplayTagsEditor/Private/SGameplayTagCombo.cpp
@@ -29,14 +29,6 @@ SGameplayTagCombo::SGameplayTagCombo()
 {
 }
 
-SGameplayTagCombo::~SGameplayTagCombo()
-{
-	if (bRegisteredForUndo)
-	{
-		GEditor->UnregisterForUndo(this);
-	}
-}
-
 void SGameplayTagCombo::Construct(const FArguments& InArgs)
 {
 	TagAttribute.Assign(*this, InArgs._Tag);
@@ -50,8 +42,6 @@ void SGameplayTagCombo::Construct(const FArguments& InArgs)
 	{
 		PropertyHandle->SetOnPropertyValueChanged(FSimpleDelegate::CreateSP(this, &SGameplayTagCombo::RefreshTagsFromProperty));
 		RefreshTagsFromProperty();
-		GEditor->RegisterForUndo(this);
-		bRegisteredForUndo = true;
 
 		if (Filter.IsEmpty())
 		{
@@ -103,22 +93,6 @@ bool SGameplayTagCombo::IsValueEnabled() const
 	return !bIsReadOnly;
 }
 
-void SGameplayTagCombo::PostUndo(bool bSuccess)
-{
-	if (bSuccess)
-	{
-		RefreshTagsFromProperty();
-	}
-}
-
-void SGameplayTagCombo::PostRedo(bool bSuccess)
-{
-	if (bSuccess)
-	{
-		RefreshTagsFromProperty();
-	}
-}
-
 FReply SGameplayTagCombo::OnEditTag() const
 {
 	FReply Reply = FReply::Handled();
@@ -364,24 +338,62 @@ bool SGameplayTagCombo::CanPaste() const
 	return PastedTag.IsValid();
 }
 
-void SGameplayTagCombo::RefreshTagsFromProperty()
+void SGameplayTagCombo::Tick(const FGeometry& AllottedGeometry, const double InCurrentTime, const float InDeltaTime)
 {
-	check(PropertyHandle.IsValid());
+	if (!PropertyHandle.IsValid()
+		|| !PropertyHandle->IsValidHandle())
+	{
+		return;
+	}
 
-	bHasMultipleValues = false;
-	TagsFromProperty.Reset();
+	// Check if cached data has changed, and update it.
+	bool bShouldUpdate = false;
 	
-	SGameplayTagPicker::EnumerateEditableTagContainersFromPropertyHandle(PropertyHandle.ToSharedRef(), [this](const FGameplayTagContainer& TagContainer)
+	TArray<const void*> RawStructData;
+	PropertyHandle->AccessRawData(RawStructData);
+
+	if (RawStructData.Num() == TagsFromProperty.Num())
 	{
-		const FGameplayTag TagFromProperty = TagContainer.IsEmpty() ? FGameplayTag() : TagContainer.First(); 
-		if (TagsFromProperty.Num() > 0 && TagsFromProperty[0] != TagFromProperty)
+		for (int32 Idx = 0; Idx < RawStructData.Num(); ++Idx)
 		{
-			bHasMultipleValues = true;
+			if (RawStructData[Idx])
+			{
+				const FGameplayTag& CurrTag = *(FGameplayTag*)RawStructData[Idx];
+				if (CurrTag != TagsFromProperty[Idx])
+				{
+					bShouldUpdate = true;
+					break;
+				}
+			}
 		}
-		TagsFromProperty.Add(TagFromProperty);
+	}
+
+	if (bShouldUpdate)
+	{
+		RefreshTagsFromProperty();
+	}
+}
 
-		return true;
-	});
+void SGameplayTagCombo::RefreshTagsFromProperty()
+{
+	if (PropertyHandle.IsValid()
+		&& PropertyHandle->IsValidHandle())
+	{
+		bHasMultipleValues = false;
+		TagsFromProperty.Reset();
+		
+		SGameplayTagPicker::EnumerateEditableTagContainersFromPropertyHandle(PropertyHandle.ToSharedRef(), [this](const FGameplayTagContainer& TagContainer)
+		{
+			const FGameplayTag TagFromProperty = TagContainer.IsEmpty() ? FGameplayTag() : TagContainer.First(); 
+			if (TagsFromProperty.Num() > 0 && TagsFromProperty[0] != TagFromProperty)
+			{
+				bHasMultipleValues = true;
+			}
+			TagsFromProperty.Add(TagFromProperty);
+
+			return true;
+		});
+	}
 }
 
 #undef LOCTEXT_NAMESPACE
diff --git a/Engine/Plugins/Editor/GameplayTagsEditor/Source/GameplayTagsEditor/Private/SGameplayTagContainerCombo.cpp b/Engine/Plugins/Editor/GameplayTagsEditor/Source/GameplayTagsEditor/Private/SGameplayTagContainerCombo.cpp
--- a/Engine/Plugins/Editor/GameplayTagsEditor/Source/GameplayTagsEditor/Private/SGameplayTagContainerCombo.cpp
+++ b/Engine/Plugins/Editor/GameplayTagsEditor/Source/GameplayTagsEditor/Private/SGameplayTagContainerCombo.cpp
@@ -35,14 +35,6 @@ SGameplayTagContainerCombo::SGameplayTagContainerCombo()
 {
 }
 
-SGameplayTagContainerCombo::~SGameplayTagContainerCombo()
-{
-	if (bRegisteredForUndo)
-	{
-		GEditor->UnregisterForUndo(this);
-	}
-}
-
 void SGameplayTagContainerCombo::Construct(const FArguments& InArgs)
 {
 	Filter = InArgs._Filter;
@@ -55,8 +47,6 @@ void SGameplayTagContainerCombo::Construct(const FArguments& InArgs)
 	{
 		PropertyHandle->SetOnPropertyValueChanged(FSimpleDelegate::CreateSP(this, &SGameplayTagContainerCombo::RefreshTagContainers));
 		RefreshTagContainers();
-		GEditor->RegisterForUndo(this);
-		bRegisteredForUndo = true;
 
 		if (Filter.IsEmpty())
 		{
@@ -511,17 +501,37 @@ FReply SGameplayTagContainerCombo::OnClearTagClicked(const FGameplayTag TagToCle
 	return FReply::Handled();
 }
 
-void SGameplayTagContainerCombo::PostUndo(bool bSuccess)
+void SGameplayTagContainerCombo::Tick(const FGeometry& AllottedGeometry, const double InCurrentTime, const float InDeltaTime)
 {
-	if (bSuccess)
+	if (!PropertyHandle.IsValid()
+		|| !PropertyHandle->IsValidHandle())
 	{
-		RefreshTagContainers();
+		return;
 	}
-}
 
-void SGameplayTagContainerCombo::PostRedo(bool bSuccess)
-{
-	if (bSuccess)
+	// Check if cached data has changed, and update it.
+	bool bShouldUpdate = false;
+		
+	TArray<const void*> RawStructData;
+	PropertyHandle->AccessRawData(RawStructData);
+
+	if (RawStructData.Num() == CachedTagContainers.Num())
+	{
+		for (int32 Idx = 0; Idx < RawStructData.Num(); ++Idx)
+		{
+			if (RawStructData[Idx])
+			{
+				const FGameplayTagContainer& TagContainer = *(FGameplayTagContainer*)RawStructData[Idx];
+				if (TagContainer != CachedTagContainers[Idx])
+				{
+					bShouldUpdate = true;
+					break;
+				}
+			}
+		}
+	}
+
+	if (bShouldUpdate)
 	{
 		RefreshTagContainers();
 	}
@@ -534,30 +544,33 @@ void SGameplayTagContainerCombo::RefreshTagContainers()
 
 	if (PropertyHandle.IsValid())
 	{
-		// From property
-		SGameplayTagPicker::EnumerateEditableTagContainersFromPropertyHandle(PropertyHandle.ToSharedRef(), [this](const FGameplayTagContainer& InTagContainer)
+		if (PropertyHandle->IsValidHandle())
 		{
-			CachedTagContainers.Add(InTagContainer);
-
-			for (auto It = InTagContainer.CreateConstIterator(); It; ++It)
+			// From property
+			SGameplayTagPicker::EnumerateEditableTagContainersFromPropertyHandle(PropertyHandle.ToSharedRef(), [this](const FGameplayTagContainer& InTagContainer)
 			{
-				const FGameplayTag Tag = *It;
-				const int32 ExistingItemIndex = TagsToEdit.IndexOfByPredicate([Tag](const TSharedPtr<FEditableItem>& Item)
-				{
-					return Item.IsValid() && Item->Tag == Tag;
-				});
-				if (ExistingItemIndex != INDEX_NONE)
-				{
-					TagsToEdit[ExistingItemIndex]->Count++;
-				}
-				else
+				CachedTagContainers.Add(InTagContainer);
+
+				for (auto It = InTagContainer.CreateConstIterator(); It; ++It)
 				{
-					TagsToEdit.Add(MakeShared<FEditableItem>(Tag));
+					const FGameplayTag Tag = *It;
+					const int32 ExistingItemIndex = TagsToEdit.IndexOfByPredicate([Tag](const TSharedPtr<FEditableItem>& Item)
+					{
+						return Item.IsValid() && Item->Tag == Tag;
+					});
+					if (ExistingItemIndex != INDEX_NONE)
+					{
+						TagsToEdit[ExistingItemIndex]->Count++;
+					}
+					else
+					{
+						TagsToEdit.Add(MakeShared<FEditableItem>(Tag));
+					}
 				}
-			}
-			
-			return true;
-		});
+				
+				return true;
+			});
+		}
 	}
 	else
 	{
diff --git a/Engine/Plugins/Editor/GameplayTagsEditor/Source/GameplayTagsEditor/Private/SGameplayTagPicker.cpp b/Engine/Plugins/Editor/GameplayTagsEditor/Source/GameplayTagsEditor/Private/SGameplayTagPicker.cpp
--- a/Engine/Plugins/Editor/GameplayTagsEditor/Source/GameplayTagsEditor/Private/SGameplayTagPicker.cpp
+++ b/Engine/Plugins/Editor/GameplayTagsEditor/Source/GameplayTagsEditor/Private/SGameplayTagPicker.cpp
@@ -75,6 +75,15 @@ bool SGameplayTagPicker::GetEditableTagContainersFromPropertyHandle(const TShare
 	});
 }
 
+SGameplayTagPicker::~SGameplayTagPicker()
+{
+	if (PostUndoRedoDelegateHandle.IsValid())
+	{
+		FEditorDelegates::PostUndoRedo.Remove(PostUndoRedoDelegateHandle);
+		PostUndoRedoDelegateHandle.Reset();
+	}
+}
+
 void SGameplayTagPicker::Construct(const FArguments& InArgs)
 {
 	TagContainers = InArgs._TagContainers;
@@ -101,6 +110,8 @@ void SGameplayTagPicker::Construct(const FArguments& InArgs)
 
 	bRestrictedTags = InArgs._RestrictedTags;
 
+	PostUndoRedoDelegateHandle = FEditorDelegates::PostUndoRedo.AddSP(this, &SGameplayTagPicker::OnPostUndoRedo);
+	
 	UGameplayTagsManager::OnEditorRefreshGameplayTagTree.AddSP(this, &SGameplayTagPicker::RefreshOnNextTick);
 	UGameplayTagsManager& Manager = UGameplayTagsManager::Get();
 
@@ -148,7 +159,7 @@ void SGameplayTagPicker::Construct(const FArguments& InArgs)
 	SettingsCombo->SetOnGetMenuContent(FOnGetContent::CreateSP(this, &SGameplayTagPicker::MakeSettingsMenu, SettingsCombo));
 
 
-	TWeakPtr<SGameplayTagPicker> WeakSelf = StaticCastWeakPtr<SGameplayTagPicker>(AsWeak());
+	TWeakPtr<SGameplayTagPicker> WeakSelf = SharedThis(this);
 	
 	TSharedRef<SWidget> Picker = 
 		SNew(SBorder)
@@ -1775,12 +1786,7 @@ void SGameplayTagPicker::SetTagContainers(TConstArrayView<FGameplayTagContainer>
 	TagContainers = InTagContainers;
 }
 
-void SGameplayTagPicker::PostUndo(bool bSuccess)
-{
-	OnRefreshTagContainers.ExecuteIfBound(*this);
-}
-
-void SGameplayTagPicker::PostRedo(bool bSuccess)
+void SGameplayTagPicker::OnPostUndoRedo()
 {
 	OnRefreshTagContainers.ExecuteIfBound(*this);
 }
diff --git a/Engine/Plugins/Editor/GameplayTagsEditor/Source/GameplayTagsEditor/Private/SGameplayTagQueryEntryBox.cpp b/Engine/Plugins/Editor/GameplayTagsEditor/Source/GameplayTagsEditor/Private/SGameplayTagQueryEntryBox.cpp
--- a/Engine/Plugins/Editor/GameplayTagsEditor/Source/GameplayTagsEditor/Private/SGameplayTagQueryEntryBox.cpp
+++ b/Engine/Plugins/Editor/GameplayTagsEditor/Source/GameplayTagsEditor/Private/SGameplayTagQueryEntryBox.cpp
@@ -31,14 +31,6 @@ SGameplayTagQueryEntryBox::SGameplayTagQueryEntryBox()
 {
 }
 
-SGameplayTagQueryEntryBox::~SGameplayTagQueryEntryBox()
-{
-	if (bRegisteredForUndo)
-	{
-		GEditor->UnregisterForUndo(this);
-	}
-}
-
 void SGameplayTagQueryEntryBox::Construct(const FArguments& InArgs)
 {
 	bIsReadOnly = InArgs._ReadOnly;
@@ -48,8 +40,6 @@ void SGameplayTagQueryEntryBox::Construct(const FArguments& InArgs)
 	if (PropertyHandle.IsValid())
 	{
 		PropertyHandle->SetOnPropertyValueChanged(FSimpleDelegate::CreateSP(this, &SGameplayTagQueryEntryBox::CacheQueryList));
-		GEditor->RegisterForUndo(this);
-		bRegisteredForUndo = true;
 		bIsReadOnly = PropertyHandle->IsEditConst();
 
 		if (Filter.IsEmpty())
@@ -130,16 +120,6 @@ void SGameplayTagQueryEntryBox::Construct(const FArguments& InArgs)
 	];
 }
 
-void SGameplayTagQueryEntryBox::PostUndo(bool bSuccess)
-{
-	CacheQueryList();
-}
-
-void SGameplayTagQueryEntryBox::PostRedo(bool bSuccess)
-{
-	CacheQueryList();
-}
-
 FText SGameplayTagQueryEntryBox::GetQueryDescText() const
 {
 	return QueryDescription;
@@ -252,22 +232,61 @@ FReply SGameplayTagQueryEntryBox::OnEditButtonClicked()
 	return FReply::Handled();
 }
 
+void SGameplayTagQueryEntryBox::Tick(const FGeometry& AllottedGeometry, const double InCurrentTime, const float InDeltaTime)
+{
+	if (!PropertyHandle.IsValid()
+		|| !PropertyHandle->IsValidHandle())
+	{
+		return;
+	}
+
+	// Check if cached data has changed, and update it.
+	bool bShouldUpdate = false;
+	
+	TArray<const void*> RawStructData;
+	PropertyHandle->AccessRawData(RawStructData);
+
+	if (RawStructData.Num() == CachedQueries.Num())
+	{
+		for (int32 Idx = 0; Idx < RawStructData.Num(); ++Idx)
+		{
+			if (RawStructData[Idx])
+			{
+				const FGameplayTagQuery& Query = *(FGameplayTagQuery*)RawStructData[Idx];
+				if (Query != CachedQueries[Idx])
+				{
+					bShouldUpdate = true;
+					break;
+				}
+			}
+		}
+	}
+
+	if (bShouldUpdate)
+	{
+		CacheQueryList();
+	}
+}
+
 void SGameplayTagQueryEntryBox::CacheQueryList()
 {
 	CachedQueries.Empty();
 
 	if (PropertyHandle.IsValid())
 	{
-		// Cache queries from the property handle. Add empty queries even if the instance data is null so that the indices match with the property handle.
-		TArray<void*> RawStructData;
-		PropertyHandle->AccessRawData(RawStructData);
-		
-		for (int32 Idx = 0; Idx < RawStructData.Num(); ++Idx)
+		if (PropertyHandle->IsValidHandle())
 		{
-			FGameplayTagQuery& Query = CachedQueries.AddDefaulted_GetRef();
-			if (RawStructData[Idx])
+			// Cache queries from the property handle. Add empty queries even if the instance data is null so that the indices match with the property handle.
+			TArray<void*> RawStructData;
+			PropertyHandle->AccessRawData(RawStructData);
+			
+			for (int32 Idx = 0; Idx < RawStructData.Num(); ++Idx)
 			{
-				Query = *(FGameplayTagQuery*)RawStructData[Idx]; 
+				FGameplayTagQuery& Query = CachedQueries.AddDefaulted_GetRef();
+				if (RawStructData[Idx])
+				{
+					Query = *(FGameplayTagQuery*)RawStructData[Idx]; 
+				}
 			}
 		}
 	}
diff --git a/Engine/Plugins/Editor/GameplayTagsEditor/Source/GameplayTagsEditor/Public/SGameplayTagCombo.h b/Engine/Plugins/Editor/GameplayTagsEditor/Source/GameplayTagsEditor/Public/SGameplayTagCombo.h
--- a/Engine/Plugins/Editor/GameplayTagsEditor/Source/GameplayTagsEditor/Public/SGameplayTagCombo.h
+++ b/Engine/Plugins/Editor/GameplayTagsEditor/Source/GameplayTagsEditor/Public/SGameplayTagCombo.h
@@ -7,7 +7,6 @@
 #include "Widgets/SCompoundWidget.h"
 #include "Widgets/DeclarativeSyntaxSupport.h"
 #include "GameplayTagContainer.h"
-#include "EditorUndoClient.h"
 #include "SGameplayTagChip.h"
 
 class IPropertyHandle;
@@ -19,7 +18,7 @@ class SGameplayTagPicker;
 /**
  * Widget for editing a Gameplay Tag.
  */
-class SGameplayTagCombo : public SCompoundWidget, public FEditorUndoClient
+class SGameplayTagCombo : public SCompoundWidget
 {
 	SLATE_DECLARE_WIDGET(SGameplayTagCombo, SCompoundWidget)
 
@@ -62,18 +61,13 @@ public:
 	SLATE_END_ARGS();
 
 	GAMEPLAYTAGSEDITOR_API SGameplayTagCombo();
-	GAMEPLAYTAGSEDITOR_API virtual ~SGameplayTagCombo() override;
 
 	GAMEPLAYTAGSEDITOR_API void Construct(const FArguments& InArgs);
 
-protected:
-	//~ Begin FEditorUndoClient Interface
-	virtual void PostUndo(bool bSuccess) override;
-	virtual void PostRedo(bool bSuccess) override;
-	//~ End FEditorUndoClient Interface
-
 private:
 
+	virtual void Tick(const FGeometry& AllottedGeometry, const double InCurrentTime, const float InDeltaTime) override;
+	
 	bool ShowClearButton() const;
 	FText GetText() const;
 	bool IsValueEnabled() const;
@@ -98,7 +92,6 @@ private:
 	bool bIsReadOnly = false;
 	FString Filter;
 	FString SettingsName;
-	bool bRegisteredForUndo = false;
 	FOnTagChanged OnTagChanged;
 	TSharedPtr<IPropertyHandle> PropertyHandle;
 	TSharedPtr<SMenuAnchor> MenuAnchor;
diff --git a/Engine/Plugins/Editor/GameplayTagsEditor/Source/GameplayTagsEditor/Public/SGameplayTagContainerCombo.h b/Engine/Plugins/Editor/GameplayTagsEditor/Source/GameplayTagsEditor/Public/SGameplayTagContainerCombo.h
--- a/Engine/Plugins/Editor/GameplayTagsEditor/Source/GameplayTagsEditor/Public/SGameplayTagContainerCombo.h
+++ b/Engine/Plugins/Editor/GameplayTagsEditor/Source/GameplayTagsEditor/Public/SGameplayTagContainerCombo.h
@@ -8,7 +8,6 @@
 #include "Widgets/Views/SListView.h"
 #include "Widgets/DeclarativeSyntaxSupport.h"
 #include "GameplayTagContainer.h"
-#include "EditorUndoClient.h"
 #include "SGameplayTagChip.h"
 
 class IPropertyHandle;
@@ -21,7 +20,7 @@ class SComboButton;
 /**
  * Widget for editing a Gameplay Tag Container.
  */
-class SGameplayTagContainerCombo : public SCompoundWidget, public FEditorUndoClient
+class SGameplayTagContainerCombo : public SCompoundWidget
 {
 	SLATE_DECLARE_WIDGET(SGameplayTagContainerCombo, SCompoundWidget)
 	
@@ -64,18 +63,13 @@ public:
 	SLATE_END_ARGS();
 
 	GAMEPLAYTAGSEDITOR_API SGameplayTagContainerCombo();
-	GAMEPLAYTAGSEDITOR_API virtual ~SGameplayTagContainerCombo() override;
 
 	GAMEPLAYTAGSEDITOR_API void Construct(const FArguments& InArgs);
 
-protected:
-	//~ Begin FEditorUndoClient Interface
-	virtual void PostUndo(bool bSuccess) override;
-	virtual void PostRedo(bool bSuccess) override;
-	//~ End FEditorUndoClient Interface
-
 private:
 
+	virtual void Tick(const FGeometry& AllottedGeometry, const double InCurrentTime, const float InDeltaTime) override;
+	
 	struct FEditableItem
 	{
 		FEditableItem() = default;
@@ -115,7 +109,6 @@ private:
 	FString Filter;
 	FString SettingsName;
 	bool bIsReadOnly = false;
-	bool bRegisteredForUndo = false;
 	FOnTagContainerChanged OnTagContainerChanged;
 	TSharedPtr<IPropertyHandle> PropertyHandle;
 	TSharedPtr<SMenuAnchor> MenuAnchor;
diff --git a/Engine/Plugins/Editor/GameplayTagsEditor/Source/GameplayTagsEditor/Public/SGameplayTagPicker.h b/Engine/Plugins/Editor/GameplayTagsEditor/Source/GameplayTagsEditor/Public/SGameplayTagPicker.h
--- a/Engine/Plugins/Editor/GameplayTagsEditor/Source/GameplayTagsEditor/Public/SGameplayTagPicker.h
+++ b/Engine/Plugins/Editor/GameplayTagsEditor/Source/GameplayTagsEditor/Public/SGameplayTagPicker.h
@@ -13,7 +13,6 @@
 #include "Widgets/Views/STableRow.h"
 #include "Widgets/Views/STreeView.h"
 #include "GameplayTagsManager.h"
-#include "EditorUndoClient.h"
 
 class IPropertyHandle;
 class SComboButton;
@@ -28,7 +27,7 @@ enum class EGameplayTagPickerMode : uint8
 };
 
 /** Widget allowing user to tag assets with gameplay tags */
-class GAMEPLAYTAGSEDITOR_API SGameplayTagPicker : public SCompoundWidget, public FSelfRegisteringEditorUndoClient
+class GAMEPLAYTAGSEDITOR_API SGameplayTagPicker : public SCompoundWidget
 {
 public:
 
@@ -114,6 +113,8 @@ public:
 	 */
 	static bool GetEditableTagContainersFromPropertyHandle(const TSharedRef<IPropertyHandle>& PropHandle, TArray<FGameplayTagContainer>& OutEditableContainers);
 
+	virtual ~SGameplayTagPicker() override;
+
 	/** Construct the actual widget */
 	void Construct(const FArguments& InArgs);
 	
@@ -152,10 +153,7 @@ private:
 		Duplicate,
 	};
 
-	// FSelfRegisteringEditorUndoClient
-	virtual void PostUndo(bool bSuccess) override;
-	virtual void PostRedo(bool bSuccess) override;
-	// ~FSelfRegisteringEditorUndoClient
+	void OnPostUndoRedo();
 
 	/** Verify the tags are all valid and if not prompt the user. */
 	void VerifyAssetTagValidity();
@@ -452,6 +450,8 @@ private:
 
 	TSharedPtr<SAddNewGameplayTagWidget> AddNewTagWidget;
 	bool bNewTagWidgetVisible = false;
+
+	FDelegateHandle PostUndoRedoDelegateHandle;
 };
 
 
diff --git a/Engine/Plugins/Editor/GameplayTagsEditor/Source/GameplayTagsEditor/Public/SGameplayTagQueryEntryBox.h b/Engine/Plugins/Editor/GameplayTagsEditor/Source/GameplayTagsEditor/Public/SGameplayTagQueryEntryBox.h
--- a/Engine/Plugins/Editor/GameplayTagsEditor/Source/GameplayTagsEditor/Public/SGameplayTagQueryEntryBox.h
+++ b/Engine/Plugins/Editor/GameplayTagsEditor/Source/GameplayTagsEditor/Public/SGameplayTagQueryEntryBox.h
@@ -6,7 +6,6 @@
 #include "Widgets/SCompoundWidget.h"
 #include "Widgets/DeclarativeSyntaxSupport.h"
 #include "GameplayTagContainer.h"
-#include "EditorUndoClient.h"
 #include "SGameplayTagQueryWidget.h"
 
 class IPropertyHandle;
@@ -15,7 +14,7 @@ class SHorizontalBox;
 /**
  * Widget for editing a Gameplay Tag Query.
  */
-class SGameplayTagQueryEntryBox : public SCompoundWidget, public FEditorUndoClient
+class SGameplayTagQueryEntryBox : public SCompoundWidget
 {
 	SLATE_DECLARE_WIDGET(SGameplayTagQueryEntryBox, SCompoundWidget)
 	
@@ -48,16 +47,11 @@ public:
 	SLATE_END_ARGS();
 
 	GAMEPLAYTAGSEDITOR_API SGameplayTagQueryEntryBox();
-	GAMEPLAYTAGSEDITOR_API virtual ~SGameplayTagQueryEntryBox() override;
 
 	GAMEPLAYTAGSEDITOR_API void Construct(const FArguments& InArgs);
 
-protected:
-	//~ Begin FEditorUndoClient Interface
-	virtual void PostUndo(bool bSuccess) override;
-	virtual void PostRedo(bool bSuccess) override;
-	//~ End FEditorUndoClient Interface
-
+	virtual void Tick(const FGeometry& AllottedGeometry, const double InCurrentTime, const float InDeltaTime) override;
+	
 private:
 
 	void CacheQueryList();
@@ -74,7 +68,6 @@ private:
 
 	bool bIsReadOnly = false;
 	FString Filter;
-	bool bRegisteredForUndo = false;
 	FOnTagQueryChanged OnTagQueryChanged;
 	TSharedPtr<IPropertyHandle> PropertyHandle;
 	TArray<FGameplayTagQuery> CachedQueries;
diff --git a/Engine/Source/Editor/PropertyEditor/Private/DetailLayoutBuilderImpl.cpp b/Engine/Source/Editor/PropertyEditor/Private/DetailLayoutBuilderImpl.cpp
--- a/Engine/Source/Editor/PropertyEditor/Private/DetailLayoutBuilderImpl.cpp
+++ b/Engine/Source/Editor/PropertyEditor/Private/DetailLayoutBuilderImpl.cpp
@@ -969,7 +969,21 @@ void FDetailLayoutBuilderImpl::Tick( float DeltaTime )
 {
 	for( auto It = TickableNodes.CreateIterator(); It; ++It )
 	{
-		(*It)->Tick( DeltaTime );
+		FDetailTreeNode* Node = *It;
+
+		// Skip ticking tree nodes which point to destroyed property nodes.
+		// This can happen when because the update order is this:
+		//	- update property nodes, calling DestroyTree(), and creating new nodes
+		//	- update layout builders (but old ones might still be referenced by the tree view) 
+		//  - tick layout builders, which includes the stale builders
+		//  - refresh tree view, which finally gets rid of the stale builders
+		TSharedPtr<FPropertyNode> PropertyNode = Node->GetPropertyNode();
+		if (PropertyNode.IsValid() && PropertyNode->IsDestroyed())
+		{
+			continue;
+		}
+		
+		Node->Tick( DeltaTime );
 	}
 }
 
diff --git a/Engine/Source/Editor/PropertyEditor/Private/PropertyHandleImpl.cpp b/Engine/Source/Editor/PropertyEditor/Private/PropertyHandleImpl.cpp
--- a/Engine/Source/Editor/PropertyEditor/Private/PropertyHandleImpl.cpp
+++ b/Engine/Source/Editor/PropertyEditor/Private/PropertyHandleImpl.cpp
@@ -2148,7 +2148,7 @@ void FPropertyValueImpl::DuplicateChild( TSharedPtr<FPropertyNode> ChildNodeToDu
 
 bool FPropertyValueImpl::HasValidPropertyNode() const
 {
-	return PropertyNode.IsValid();
+	return PropertyNode.IsValid() && !PropertyNode.Pin()->IsDestroyed();
 }
 
 FText FPropertyValueImpl::GetDisplayName() const
diff --git a/Engine/Source/Editor/PropertyEditor/Private/PropertyNode.cpp b/Engine/Source/Editor/PropertyEditor/Private/PropertyNode.cpp
--- a/Engine/Source/Editor/PropertyEditor/Private/PropertyNode.cpp
+++ b/Engine/Source/Editor/PropertyEditor/Private/PropertyNode.cpp
@@ -90,7 +90,7 @@ FPropertyNode::~FPropertyNode()
 void FPropertyNode::InitNode(const FPropertyNodeInitParams& InitParams)
 {
 	//Dismantle the previous tree
-	DestroyTree();
+	DestroyTree(/*bInDestroySelf*/false);
 
 	//tree hierarchy
 	check(InitParams.ParentNode.Get() != this);
@@ -2667,9 +2667,33 @@ bool FPropertyNode::IsChildOfFavorite (void) const
  */
 void FPropertyNode::DestroyTree(const bool bInDestroySelf)
 {
+	if (bInDestroySelf)
+	{
+		bIsDestroyed = true;
+	}
+
+	// Marks all the child nodes as destroyed.
+	// We cannot call DestroyTree() recursively since some UI code that gets executed
+	// on the destroyed nodes (due to unfortunate update order) assume that child nodes are always available.
+	for (TSharedPtr<FPropertyNode>& ChildNode : ChildNodes)
+	{
+		ChildNode->MarkDestroyedRecursive();
+	}
+	
 	ChildNodes.Empty();
 }
 
+void FPropertyNode::MarkDestroyedRecursive()
+{
+	bIsDestroyed = true;
+	
+	for (TSharedPtr<FPropertyNode>& ChildNode : ChildNodes)
+	{
+		check(ChildNode.IsValid());
+		ChildNode->MarkDestroyedRecursive();
+	}
+}
+
 /**
  * Marks windows as visible based on the filter strings (EVEN IF normally NOT EXPANDED)
  */
@@ -3237,6 +3261,11 @@ bool FPropertyNode::IsIgnoringInstancedReference() const
 	return bIgnoreInstancedReference;
 }
 
+bool FPropertyNode::IsDestroyed() const
+{
+	return bIsDestroyed;
+}
+
 FDelegateHandle FPropertyNode::SetOnRebuildChildren(const FSimpleDelegate& InOnRebuildChildren)
 {
 	return OnRebuildChildrenEvent.Add(InOnRebuildChildren);
diff --git a/Engine/Source/Editor/PropertyEditor/Private/PropertyNode.h b/Engine/Source/Editor/PropertyEditor/Private/PropertyNode.h
--- a/Engine/Source/Editor/PropertyEditor/Private/PropertyNode.h
+++ b/Engine/Source/Editor/PropertyEditor/Private/PropertyNode.h
@@ -1070,6 +1070,14 @@ public:
 	/** Queries whether the node would like to ignore CPF_InstancedReference semantics */
 	bool IsIgnoringInstancedReference() const;
 
+	/** Return true if DestroyTree() has been called on this node. */
+	bool IsDestroyed() const;
+
+	/**
+	 * Sets bIsDestroyed on all nodes within the hierarchy
+	 */
+	void MarkDestroyedRecursive();
+
 protected:
 	TSharedRef<FEditPropertyChain> BuildPropertyChain( FProperty* PropertyAboutToChange ) const;
 	TSharedRef<FEditPropertyChain> BuildPropertyChain( FProperty* PropertyAboutToChange, const TSet<UObject*>& InAffectedArchetypeInstances ) const;
@@ -1236,6 +1244,9 @@ protected:
 	/** Set to true when we want to ignore CPF_InstancedReference */
 	bool bIgnoreInstancedReference;
 
+	/** If true, DestroyTree() has been called on the node. */
+	bool bIsDestroyed = false;
+	
 	/** An array of restrictions limiting this property's potential values in property editors.*/
 	TArray<TSharedRef<const FPropertyRestriction>> Restrictions;
 
