Adding option to paste single tags into an empty GameplayTagContainer

diff --git a/Engine/Plugins/Editor/GameplayTagsEditor/Source/GameplayTagsEditor/Private/SGameplayTagContainerCombo.cpp b/Engine/Plugins/Editor/GameplayTagsEditor/Source/GameplayTagsEditor/Private/SGameplayTagContainerCombo.cpp
--- a/Engine/Plugins/Editor/GameplayTagsEditor/Source/GameplayTagsEditor/Private/SGameplayTagContainerCombo.cpp
+++ b/Engine/Plugins/Editor/GameplayTagsEditor/Source/GameplayTagsEditor/Private/SGameplayTagContainerCombo.cpp
@@ -104,6 +104,15 @@ void SGameplayTagContainerCombo::Construct(const FArguments& InArgs)
 					SNew(SBorder)
 					.Padding(FMargin(6,2))
 					.BorderImage(FGameplayTagStyle::GetBrush("GameplayTags.Container"))
+					.OnMouseButtonDown_Lambda([WeakSelf](const FGeometry& MyGeometry, const FPointerEvent& MouseEvent)
+					{
+						const TSharedPtr<SGameplayTagContainerCombo> Self = WeakSelf.Pin();
+						if (Self.IsValid() && MouseEvent.IsMouseButtonDown(EKeys::RightMouseButton) && Self->TagsToEdit.Num() <= 0)
+						{
+							return Self->OnEmptyMenu(MouseEvent);
+						}
+						return FReply::Unhandled();
+					})
 					[
 						SNew(SHorizontalBox)
 
@@ -285,6 +294,23 @@ FReply SGameplayTagContainerCombo::OnTagMenu(const FPointerEvent& MouseEvent, co
 
 }
 
+FReply SGameplayTagContainerCombo::OnEmptyMenu(const FPointerEvent& MouseEvent)
+{
+	FMenuBuilder MenuBuilder(/*bShouldCloseWindowAfterMenuSelection=*/ true, /*CommandList=*/ nullptr);
+	
+	MenuBuilder.AddMenuEntry(
+	NSLOCTEXT("PropertyView", "PasteProperty", "Paste"),
+		LOCTEXT("GameplayTagContainerCombo_PasteTagTooltip", "Paste tags from clipboard."),
+		FSlateIcon(FAppStyle::GetAppStyleSetName(), "GenericCommands.Paste"),
+		FUIAction(FExecuteAction::CreateSP(this, &SGameplayTagContainerCombo::OnPasteTag), FCanExecuteAction::CreateSP(this, &SGameplayTagContainerCombo::CanPaste)));
+
+	// Spawn context menu
+	const FWidgetPath WidgetPath = MouseEvent.GetEventPath() != nullptr ? *MouseEvent.GetEventPath() : FWidgetPath();
+	FSlateApplication::Get().PushMenu(AsShared(), WidgetPath, MenuBuilder.MakeWidget(), MouseEvent.GetScreenSpacePosition(), FPopupTransitionEffect(FPopupTransitionEffect::ContextMenu));
+	
+	return FReply::Handled();
+}
+
 void SGameplayTagContainerCombo::OnSearchForAnyReferences() const
 {
 	if (FEditorDelegates::OnOpenReferenceViewer.IsBound()
diff --git a/Engine/Plugins/Editor/GameplayTagsEditor/Source/GameplayTagsEditor/Public/SGameplayTagContainerCombo.h b/Engine/Plugins/Editor/GameplayTagsEditor/Source/GameplayTagsEditor/Public/SGameplayTagContainerCombo.h
--- a/Engine/Plugins/Editor/GameplayTagsEditor/Source/GameplayTagsEditor/Public/SGameplayTagContainerCombo.h
+++ b/Engine/Plugins/Editor/GameplayTagsEditor/Source/GameplayTagsEditor/Public/SGameplayTagContainerCombo.h
@@ -91,6 +91,7 @@ private:
 	TSharedRef<SWidget> OnGetMenuContent();
 	
 	FReply OnTagMenu(const FPointerEvent& MouseEvent, const FGameplayTag GameplayTag);
+	FReply OnEmptyMenu(const FPointerEvent& MouseEvent);
 	FReply OnEditClicked(const FGameplayTag TagToHilight);
 	FReply OnClearAllClicked();
 	FReply OnClearTagClicked(const FGameplayTag TagToClear);
