[UE] Improve GameplayTags.DumpTagList
Added an output for Sources and References for each GameplayTag
Added Columns for showing if tag has Native or DefaultIni source

#tests PIE

diff --git a/Engine/Plugins/Editor/GameplayTagsEditor/Source/GameplayTagsEditor/Private/GameplayTagsEditorModule.cpp b/Engine/Plugins/Editor/GameplayTagsEditor/Source/GameplayTagsEditor/Private/GameplayTagsEditorModule.cpp
--- a/Engine/Plugins/Editor/GameplayTagsEditor/Source/GameplayTagsEditor/Private/GameplayTagsEditorModule.cpp
+++ b/Engine/Plugins/Editor/GameplayTagsEditor/Source/GameplayTagsEditor/Private/GameplayTagsEditorModule.cpp
@@ -922,8 +922,12 @@ public:
 		UGameplayTagsManager& Manager = UGameplayTagsManager::Get();
 
 		TArray<FString> ReportLines;
+		TArray<FString> ReportReferencers;
+		TArray<FString> ReportSources;
 
-		ReportLines.Add(TEXT("Tag,Reference Count,Source,Comment"));
+		ReportLines.Add(TEXT("Tag,Explicit,HasNativeSource,HasConfigSource,Reference Count,Sources Count,Comment"));
+		ReportReferencers.Add(TEXT("Asset,Tag"));
+		ReportSources.Add(TEXT("Source,Tag"));
 
 		FGameplayTagContainer AllTags;
 		Manager.RequestAllGameplayTags(AllTags, true);
@@ -942,15 +946,39 @@ public:
 			AssetRegistryModule.Get().GetReferencers(TagId, Referencers, UE::AssetRegistry::EDependencyCategory::SearchableName);
 
 			FString Comment;
-			FName TagSource;
+			TArray<FName> TagSources;
 			bool bExplicit, bRestricted, bAllowNonRestrictedChildren;
 
-			Manager.GetTagEditorData(Tag.GetTagName(), Comment, TagSource, bExplicit, bRestricted, bAllowNonRestrictedChildren);
+			Manager.GetTagEditorData(Tag.GetTagName(), Comment, TagSources, bExplicit, bRestricted, bAllowNonRestrictedChildren);
 
-			ReportLines.Add(FString::Printf(TEXT("%s,%d,%s,%s"), *Tag.ToString(), Referencers.Num(), *TagSource.ToString(), *Comment));
+			bool bHasNative = TagSources.Contains(FGameplayTagSource::GetNativeName());
+			bool bHasConfigIni = TagSources.Contains(FGameplayTagSource::GetDefaultName());
+
+			FString TagName = Tag.ToString();
+
+			ReportLines.Add(FString::Printf(TEXT("%s,%s,%s,%s,%d,%d,\"%s\""),
+				*TagName,
+				bExplicit ? TEXT("true") : TEXT("false"),
+				bHasNative ? TEXT("true") : TEXT("false"),
+				bHasConfigIni ? TEXT("true") : TEXT("false"),
+				Referencers.Num(),
+				TagSources.Num(),
+				*Comment));
+
+			for (const FAssetIdentifier& Referencer : Referencers)
+			{
+				ReportReferencers.Add(FString::Printf(TEXT("%s,%s"), *Referencer.ToString(), *TagName));
+			}
+
+			for (const FName& TagSource : TagSources)
+			{
+				ReportSources.Add(FString::Printf(TEXT("%s,%s"), *TagSource.ToString(), *TagName));
+			}
 		}
 
 		WriteCustomReport(TEXT("TagList.csv"), ReportLines);
+		WriteCustomReport(TEXT("TagReferencesList.csv"), ReportReferencers);
+		WriteCustomReport(TEXT("TagSourcesList.csv"), ReportSources);
 	}
 
 	FDelegateHandle AssetImportHandle;
@@ -962,7 +990,8 @@ public:
 
 static FAutoConsoleCommand CVarDumpTagList(
 	TEXT("GameplayTags.DumpTagList"),
-	TEXT("Writes out a csv with all tags to Reports/TagList.csv"),
+	TEXT("Writes out a csvs with all tags to Reports/TagList.csv, ")
+	TEXT("Reports/TagReferencesList.csv and Reports/TagSourcesList.csv"),
 	FConsoleCommandDelegate::CreateStatic(FGameplayTagsEditorModule::DumpTagList),
 	ECVF_Cheat);
 
