Objective Debugger - Filter by [REDACTED_PERSON]

diff --git a/Engine/Plugins/Editor/GameplayTagsEditor/Source/GameplayTagsEditor/Private/SGameplayTagQueryEntryBox.cpp b/Engine/Plugins/Editor/GameplayTagsEditor/Source/GameplayTagsEditor/Private/SGameplayTagQueryEntryBox.cpp
--- a/Engine/Plugins/Editor/GameplayTagsEditor/Source/GameplayTagsEditor/Private/SGameplayTagQueryEntryBox.cpp
+++ b/Engine/Plugins/Editor/GameplayTagsEditor/Source/GameplayTagsEditor/Private/SGameplayTagQueryEntryBox.cpp
@@ -6,6 +6,7 @@
 #include "Widgets/Input/SButton.h"
 #include "Widgets/Images/SImage.h"
 #include "ScopedTransaction.h"
+#include "SGameplayTagQueryWidget.h"
 #include "Editor.h"
 #include "GameplayTagEditorUtilities.h"
 #include "GameplayTagsManager.h"
diff --git a/Engine/Plugins/Editor/GameplayTagsEditor/Source/GameplayTagsEditor/Public/SGameplayTagQueryEntryBox.h b/Engine/Plugins/Editor/GameplayTagsEditor/Source/GameplayTagsEditor/Public/SGameplayTagQueryEntryBox.h
--- a/Engine/Plugins/Editor/GameplayTagsEditor/Source/GameplayTagsEditor/Public/SGameplayTagQueryEntryBox.h
+++ b/Engine/Plugins/Editor/GameplayTagsEditor/Source/GameplayTagsEditor/Public/SGameplayTagQueryEntryBox.h
@@ -6,10 +6,10 @@
 #include "Widgets/SCompoundWidget.h"
 #include "Widgets/DeclarativeSyntaxSupport.h"
 #include "GameplayTagContainer.h"
-#include "SGameplayTagQueryWidget.h"
 
 class IPropertyHandle;
 class SHorizontalBox;
+class SGameplayTagQueryWidget;
 
 /**
  * Widget for editing a Gameplay Tag Query.
