Fixed FGameplayTag and FGameplayTagContainer pin DefaultValue modifications not being undoable/redoable. Now properly creates a transaction in Undo history. Additionally, the pin "Reset to Default Value" action now also creates an undoable transaction.

diff --git a/Engine/Plugins/Editor/GameplayTagsEditor/Source/GameplayTagsEditor/Private/SGameplayTagContainerGraphPin.cpp b/Engine/Plugins/Editor/GameplayTagsEditor/Source/GameplayTagsEditor/Private/SGameplayTagContainerGraphPin.cpp
--- a/Engine/Plugins/Editor/GameplayTagsEditor/Source/GameplayTagsEditor/Private/SGameplayTagContainerGraphPin.cpp
+++ b/Engine/Plugins/Editor/GameplayTagsEditor/Source/GameplayTagsEditor/Private/SGameplayTagContainerGraphPin.cpp
@@ -35,6 +35,7 @@ void SGameplayTagContainerGraphPin::OnTagContainerChanged(const FGameplayTagCont
 	if (!CurrentDefaultValue.Equals(TagContainerString))
 	{
 		const FScopedTransaction Transaction(LOCTEXT("ChangeDefaultValue", "Change Pin Default Value"));
+		GraphPinObj->Modify();
 		GraphPinObj->GetSchema()->TrySetDefaultValue(*GraphPinObj, TagContainerString);
 	}
 }
diff --git a/Engine/Plugins/Editor/GameplayTagsEditor/Source/GameplayTagsEditor/Private/SGameplayTagGraphPin.cpp b/Engine/Plugins/Editor/GameplayTagsEditor/Source/GameplayTagsEditor/Private/SGameplayTagGraphPin.cpp
--- a/Engine/Plugins/Editor/GameplayTagsEditor/Source/GameplayTagsEditor/Private/SGameplayTagGraphPin.cpp
+++ b/Engine/Plugins/Editor/GameplayTagsEditor/Source/GameplayTagsEditor/Private/SGameplayTagGraphPin.cpp
@@ -35,6 +35,7 @@ void SGameplayTagGraphPin::OnTagChanged(const FGameplayTag NewTag)
 	if (!CurrentDefaultValue.Equals(TagString))
 	{
 		const FScopedTransaction Transaction(LOCTEXT("ChangeDefaultValue", "Change Pin Default Value"));
+		GraphPinObj->Modify();
 		GraphPinObj->GetSchema()->TrySetDefaultValue(*GraphPinObj, TagString);
 	}
 }
diff --git a/Engine/Plugins/Editor/GameplayTagsEditor/Source/GameplayTagsEditor/Private/SGameplayTagQueryGraphPin.cpp b/Engine/Plugins/Editor/GameplayTagsEditor/Source/GameplayTagsEditor/Private/SGameplayTagQueryGraphPin.cpp
--- a/Engine/Plugins/Editor/GameplayTagsEditor/Source/GameplayTagsEditor/Private/SGameplayTagQueryGraphPin.cpp
+++ b/Engine/Plugins/Editor/GameplayTagsEditor/Source/GameplayTagsEditor/Private/SGameplayTagQueryGraphPin.cpp
@@ -25,6 +25,7 @@ void SGameplayTagQueryGraphPin::OnTagQueryChanged(const FGameplayTagQuery& NewTa
 	{
 		// Set Pin Data
 		const FScopedTransaction Transaction(NSLOCTEXT("GraphEditor", "ChangePinValue", "Change Pin Value"));
+		GraphPinObj->Modify();
 		GraphPinObj->GetSchema()->TrySetDefaultValue(*GraphPinObj, TagQueryExportText);
 	}
 }
diff --git a/Engine/Source/Editor/Kismet/Private/BlueprintEditor.cpp b/Engine/Source/Editor/Kismet/Private/BlueprintEditor.cpp
--- a/Engine/Source/Editor/Kismet/Private/BlueprintEditor.cpp
+++ b/Engine/Source/Editor/Kismet/Private/BlueprintEditor.cpp
@@ -5227,6 +5227,7 @@ void FBlueprintEditor::OnResetPinToDefaultValue()
 		const FScopedTransaction Transaction(LOCTEXT("ResetPinToDefaultValue", "Reset Pin To Default Value"));
 
 		const UEdGraphSchema_K2* K2Schema = GetDefault<UEdGraphSchema_K2>();
+		TargetPin->Modify();
 		K2Schema->ResetPinToAutogeneratedDefaultValue(TargetPin);
 	}
 }
