GameplayTags: Fix bugs getting and passing category filters
- Remove unnecessary call to UGameplayTagsManager::StaticCategoriesMetaFromPropertyHandle() in SGameplayTagQueryEntryBox::OnEditButtonClicked() as the filter is already set up in the constructor
- Added missing assignment of member Filter in SGameplayTagQueryWidget::Construct
- Added missing assignment of member Filter in SGameplayTagQueryEntryBox::Construct

diff --git a/Engine/Plugins/Editor/GameplayTagsEditor/Source/GameplayTagsEditor/Private/SGameplayTagQueryEntryBox.cpp b/Engine/Plugins/Editor/GameplayTagsEditor/Source/GameplayTagsEditor/Private/SGameplayTagQueryEntryBox.cpp
--- a/Engine/Plugins/Editor/GameplayTagsEditor/Source/GameplayTagsEditor/Private/SGameplayTagQueryEntryBox.cpp
+++ b/Engine/Plugins/Editor/GameplayTagsEditor/Source/GameplayTagsEditor/Private/SGameplayTagQueryEntryBox.cpp
@@ -34,6 +34,7 @@ SGameplayTagQueryEntryBox::SGameplayTagQueryEntryBox()
 
 void SGameplayTagQueryEntryBox::Construct(const FArguments& InArgs)
 {
+	Filter = InArgs._Filter;
 	bIsReadOnly = InArgs._ReadOnly;
 	OnTagQueryChanged = InArgs._OnTagQueryChanged;
 	PropertyHandle = InArgs._PropertyHandle;
@@ -195,8 +196,6 @@ FReply SGameplayTagQueryEntryBox::OnEditButtonClicked()
 	
 	if (PropertyHandle.IsValid())
 	{
-		Args.Filter = UGameplayTagsManager::StaticGetCategoriesMetaFromPropertyHandle(PropertyHandle);
-
 		TArray<UObject*> OuterObjects;
 		PropertyHandle->GetOuterObjects(OuterObjects);
 
diff --git a/Engine/Plugins/Editor/GameplayTagsEditor/Source/GameplayTagsEditor/Private/SGameplayTagQueryWidget.cpp b/Engine/Plugins/Editor/GameplayTagsEditor/Source/GameplayTagsEditor/Private/SGameplayTagQueryWidget.cpp
--- a/Engine/Plugins/Editor/GameplayTagsEditor/Source/GameplayTagsEditor/Private/SGameplayTagQueryWidget.cpp
+++ b/Engine/Plugins/Editor/GameplayTagsEditor/Source/GameplayTagsEditor/Private/SGameplayTagQueryWidget.cpp
@@ -22,6 +22,7 @@ void SGameplayTagQueryWidget::Construct(const FArguments& InArgs, const TArray<F
 
 	bReadOnly = InArgs._ReadOnly;
 	OnQueriesCommitted = InArgs._OnQueriesCommitted;
+	Filter = InArgs._Filter;
 
 	UGameplayTagsManager::Get().OnGetCategoriesMetaFromPropertyHandle.AddSP(this, &SGameplayTagQueryWidget::OnGetCategoriesMetaFromPropertyHandle);
 
