PR #11845: GameplayTags: Fix 'Categories' UPARAM meta data not working for FGameplayTagQuery pins

diff --git a/Engine/Plugins/Editor/GameplayTagsEditor/Source/GameplayTagsEditor/Private/SGameplayTagQueryGraphPin.cpp b/Engine/Plugins/Editor/GameplayTagsEditor/Source/GameplayTagsEditor/Private/SGameplayTagQueryGraphPin.cpp
--- a/Engine/Plugins/Editor/GameplayTagsEditor/Source/GameplayTagsEditor/Private/SGameplayTagQueryGraphPin.cpp
+++ b/Engine/Plugins/Editor/GameplayTagsEditor/Source/GameplayTagsEditor/Private/SGameplayTagQueryGraphPin.cpp
@@ -36,6 +36,7 @@ TSharedRef<SWidget>	SGameplayTagQueryGraphPin::GetDefaultValueWidget()
 
 	return SNew(SGameplayTagQueryEntryBox)
 		.Visibility(this, &SGraphPin::GetDefaultValueVisibility)
+		.Filter(UE::GameplayTags::EditorUtilities::ExtractTagFilterStringFromGraphPin(GraphPinObj))
 		.DescriptionMaxWidth(250.0f)
 		.TagQuery(this, &SGameplayTagQueryGraphPin::GetTagQuery)
 		.OnTagQueryChanged(this, &SGameplayTagQueryGraphPin::OnTagQueryChanged);
