Gameplay Tag Editor: Fix query window does not update the property after the detail panel has been updated.
Also disable the window if the underlying properly completely disappears.

diff --git a/Engine/Plugins/Editor/GameplayTagsEditor/Source/GameplayTagsEditor/Private/SGameplayTagQueryEntryBox.cpp b/Engine/Plugins/Editor/GameplayTagsEditor/Source/GameplayTagsEditor/Private/SGameplayTagQueryEntryBox.cpp
--- a/Engine/Plugins/Editor/GameplayTagsEditor/Source/GameplayTagsEditor/Private/SGameplayTagQueryEntryBox.cpp
+++ b/Engine/Plugins/Editor/GameplayTagsEditor/Source/GameplayTagsEditor/Private/SGameplayTagQueryEntryBox.cpp
@@ -39,6 +39,8 @@ void SGameplayTagQueryEntryBox::Construct(const FArguments& InArgs)
 	OnTagQueryChanged = InArgs._OnTagQueryChanged;
 	PropertyHandle = InArgs._PropertyHandle;
 
+	CacheQueryList();
+
 	if (PropertyHandle.IsValid())
 	{
 		PropertyHandle->SetOnPropertyValueChanged(FSimpleDelegate::CreateSP(this, &SGameplayTagQueryEntryBox::CacheQueryList));
@@ -48,14 +50,15 @@ void SGameplayTagQueryEntryBox::Construct(const FArguments& InArgs)
 		{
 			Filter = UGameplayTagsManager::Get().GetCategoriesMetaFromPropertyHandle(PropertyHandle);
 		}
+
+		// Let's check if there's already an open editor window for this property and rebuild that
+		QueryWidget = UE::GameplayTags::Editor::TrySyncGameplayTagQueryWidget(BuildWindowArgs());
 	}
 	else
 	{
 		TagQueryAttribute.Assign(*this, InArgs._TagQuery);
 	}
 
-	CacheQueryList();
-
 	ChildSlot
 	[
 		SAssignNew(WidgetContainer, SHorizontalBox)
@@ -185,7 +188,7 @@ EVisibility SGameplayTagQueryEntryBox::GetQueryDescVisibility() const
 	return HasAnyValidQueries() == true ? EVisibility::Visible : EVisibility::Collapsed;
 }
 
-FReply SGameplayTagQueryEntryBox::OnEditButtonClicked()
+FGameplayTagQueryWindowArgs SGameplayTagQueryEntryBox::BuildWindowArgs()
 {
 	FGameplayTagQueryWindowArgs Args;
 	Args.OnQueriesCommitted = SGameplayTagQueryWidget::FOnQueriesCommitted::CreateSP(this, &SGameplayTagQueryEntryBox::OnQueriesCommitted);
@@ -196,6 +199,8 @@ FReply SGameplayTagQueryEntryBox::OnEditButtonClicked()
 	
 	if (PropertyHandle.IsValid())
 	{
+		Args.Property = PropertyHandle->GetProperty();
+
 		TArray<UObject*> OuterObjects;
 		PropertyHandle->GetOuterObjects(OuterObjects);
 
@@ -227,7 +232,12 @@ FReply SGameplayTagQueryEntryBox::OnEditButtonClicked()
 		Args.Title = LOCTEXT("GameplayTagQueryEntryBox_WidgetTitle", "Tag Editor");
 	}
 
-	QueryWidget = UE::GameplayTags::Editor::OpenGameplayTagQueryWindow(Args);
+	return Args;
+}
+
+FReply SGameplayTagQueryEntryBox::OnEditButtonClicked()
+{
+	QueryWidget = UE::GameplayTags::Editor::OpenGameplayTagQueryWindow(BuildWindowArgs());
 	
 	return FReply::Handled();
 }
diff --git a/Engine/Plugins/Editor/GameplayTagsEditor/Source/GameplayTagsEditor/Private/SGameplayTagQueryWidget.cpp b/Engine/Plugins/Editor/GameplayTagsEditor/Source/GameplayTagsEditor/Private/SGameplayTagQueryWidget.cpp
--- a/Engine/Plugins/Editor/GameplayTagsEditor/Source/GameplayTagsEditor/Private/SGameplayTagQueryWidget.cpp
+++ b/Engine/Plugins/Editor/GameplayTagsEditor/Source/GameplayTagsEditor/Private/SGameplayTagQueryWidget.cpp
@@ -42,6 +42,9 @@ void SGameplayTagQueryWidget::Construct(const FArguments& InArgs, const TArray<F
 	FPropertyEditorModule& PropertyModule = FModuleManager::LoadModuleChecked<FPropertyEditorModule>("PropertyEditor");
 	Details = PropertyModule.CreateDetailView(ViewArgs);
 	Details->SetObject(EditableGameplayTagQuery);
+	Details->SetEnabled(TAttribute<bool>::Create(
+		TAttribute<bool>::FGetter::CreateSP(this, &SGameplayTagQueryWidget::IsBoundToWindow)
+	));
 
 	ChildSlot
 	[
@@ -82,6 +85,7 @@ void SGameplayTagQueryWidget::Construct(const FArguments& InArgs, const TArray<F
 				[
 					// cancel button
 					SNew(SButton)
+						.IsEnabled(this, &SGameplayTagQueryWidget::IsBoundToWindow)
 						.ContentPadding(FAppStyle::Get().GetMargin("StandardDialog.ContentPadding") )
 						.HAlign(HAlign_Center)
 						.OnClicked(this, &SGameplayTagQueryWidget::OnCancelClicked)
@@ -125,6 +129,11 @@ UEditableGameplayTagQuery* SGameplayTagQueryWidget::CreateEditableQuery(FGamepla
 	return AnEditableQuery;
 }
 
+bool SGameplayTagQueryWidget::IsBoundToWindow() const
+{
+	return OnQueriesCommitted.IsBound();
+}
+
 SGameplayTagQueryWidget::~SGameplayTagQueryWidget()
 {
 	// clean up our temp editing uobjects
@@ -205,6 +214,7 @@ namespace UE::GameplayTags::Editor
 
 static TWeakPtr<SGameplayTagQueryWidget> GlobalTagQueryWidget;
 static TWeakPtr<SWindow> GlobalTagQueryWidgetWindow;
+static TWeakFieldPtr<FProperty> GlobalTagQueryWidgetProperty;
 
 TWeakPtr<SGameplayTagQueryWidget> OpenGameplayTagQueryWindow(const FGameplayTagQueryWindowArgs& Args)
 {
@@ -262,6 +272,7 @@ TWeakPtr<SGameplayTagQueryWidget> OpenGameplayTagQueryWindow(const FGameplayTagQ
 
 	GlobalTagQueryWidget = QueryWidget;
 	GlobalTagQueryWidgetWindow = Window;
+	GlobalTagQueryWidgetProperty = Args.Property;
 
 	auto CloseWindow = [WeakQueryWidget = GlobalTagQueryWidget]()
 	{
@@ -286,8 +297,43 @@ void CloseGameplayTagQueryWindow(TWeakPtr<SGameplayTagQueryWidget> QueryWidget)
 
 	GlobalTagQueryWidget = nullptr;
 	GlobalTagQueryWidgetWindow = nullptr;
+	GlobalTagQueryWidgetProperty.Reset();
 }
 
+TWeakPtr<SGameplayTagQueryWidget> TrySyncGameplayTagQueryWidget(const FGameplayTagQueryWindowArgs& Args)
+{
+	// If the property does not match the currently open one, then there's nothing to sync
+	if (!Args.Property.IsValid() || !Args.Property.HasSameIndexAndSerialNumber(GlobalTagQueryWidgetProperty))
+	{
+		return nullptr;
+	}
+
+	if (const TSharedPtr<SWindow> WidgetWindow = GlobalTagQueryWidgetWindow.Pin())
+	{
+		// If the window title is different, then it belongs to a different asset or property
+		if (!WidgetWindow->GetTitle().EqualTo(Args.Title))
+		{
+			return nullptr;
+		}
+	}
+
+	if (const TSharedPtr<SGameplayTagQueryWidget> QueryWidget = GlobalTagQueryWidget.Pin())
+	{
+		const bool bDetailsMatch = QueryWidget->IsReadOnly() == Args.bReadOnly
+			&& QueryWidget->GetFilter() == Args.Filter
+			&& QueryWidget->GetTagQueries() == Args.EditableQueries;
+
+		// Only sync if properties are the same
+		if (bDetailsMatch)
+		{
+			GlobalTagQueryWidgetProperty = Args.Property;
+			QueryWidget->SetOnQueriesCommitted(Args.OnQueriesCommitted);
+			return GlobalTagQueryWidget;
+		}
+	}
+
+	return nullptr;
+}
 } // UE::GameplayTags::Editor
 
 #undef LOCTEXT_NAMESPACE
diff --git a/Engine/Plugins/Editor/GameplayTagsEditor/Source/GameplayTagsEditor/Private/SGameplayTagQueryWidget.h b/Engine/Plugins/Editor/GameplayTagsEditor/Source/GameplayTagsEditor/Private/SGameplayTagQueryWidget.h
--- a/Engine/Plugins/Editor/GameplayTagsEditor/Source/GameplayTagsEditor/Private/SGameplayTagQueryWidget.h
+++ b/Engine/Plugins/Editor/GameplayTagsEditor/Source/GameplayTagsEditor/Private/SGameplayTagQueryWidget.h
@@ -4,6 +4,7 @@
 
 #include "Widgets/SCompoundWidget.h"
 #include "Misc/NotifyHook.h"
+#include "UObject/WeakFieldPtr.h"
 
 class IPropertyHandle;
 struct FGameplayTagQuery;
@@ -38,6 +39,14 @@ public:
 
 	void SetOnOk(FSimpleDelegate InOnOk) { OnOk = InOnOk; }
 	void SetOnCancel(FSimpleDelegate InOnCancel) { OnCancel = InOnCancel; }
+	void SetOnQueriesCommitted(FOnQueriesCommitted InOnQueriesCommitted) { OnQueriesCommitted = InOnQueriesCommitted; }
+
+	bool IsReadOnly() const { return bReadOnly; }
+	const TArray<FGameplayTagQuery>& GetTagQueries() { return TagQueries; }
+	const FString& GetFilter() { return Filter; }
+
+	/** Returns true if this widget is currently bound to an SWindow (i.e. it's possible to commit the changes) */
+	bool IsBoundToWindow() const;
 	
 	virtual ~SGameplayTagQueryWidget() override;
 
@@ -108,10 +117,14 @@ struct FGameplayTagQueryWindowArgs
 
 	/** (Optional) Widget to use to position the query window. */
 	TSharedPtr<SWidget> AnchorWidget;
+
+	/** (Optional) The underlying edited property. Used for re-assigning the window upon a detail panel refresh. */
+	TWeakFieldPtr<FProperty> Property;
 };
 
 namespace UE::GameplayTags::Editor
 {
 	TWeakPtr<SGameplayTagQueryWidget> OpenGameplayTagQueryWindow(const FGameplayTagQueryWindowArgs& Args);
 	void CloseGameplayTagQueryWindow(TWeakPtr<SGameplayTagQueryWidget> QueryWidget);
+	TWeakPtr<SGameplayTagQueryWidget> TrySyncGameplayTagQueryWidget(const FGameplayTagQueryWindowArgs& Args);
 };
\ No newline at end of file
diff --git a/Engine/Plugins/Editor/GameplayTagsEditor/Source/GameplayTagsEditor/Public/SGameplayTagQueryEntryBox.h b/Engine/Plugins/Editor/GameplayTagsEditor/Source/GameplayTagsEditor/Public/SGameplayTagQueryEntryBox.h
--- a/Engine/Plugins/Editor/GameplayTagsEditor/Source/GameplayTagsEditor/Public/SGameplayTagQueryEntryBox.h
+++ b/Engine/Plugins/Editor/GameplayTagsEditor/Source/GameplayTagsEditor/Public/SGameplayTagQueryEntryBox.h
@@ -55,6 +55,7 @@ public:
 private:
 
 	void CacheQueryList();
+	struct FGameplayTagQueryWindowArgs BuildWindowArgs();
 	FReply OnEditButtonClicked();
 	FReply OnClearAllButtonClicked();
 	EVisibility GetQueryDescVisibility() const;
