PR #10616: Fix not being able to edit Gameplay Tag Queries from modal windows

#tests Used GEditor->EditorAddModalWindow on our VerbDebugger and tested if I can edit the filter without this change (I can't) and then with it (I can).

diff --git a/Engine/Plugins/Editor/GameplayTagsEditor/Source/GameplayTagsEditor/Private/SGameplayTagQueryWidget.cpp b/Engine/Plugins/Editor/GameplayTagsEditor/Source/GameplayTagsEditor/Private/SGameplayTagQueryWidget.cpp
--- a/Engine/Plugins/Editor/GameplayTagsEditor/Source/GameplayTagsEditor/Private/SGameplayTagQueryWidget.cpp
+++ b/Engine/Plugins/Editor/GameplayTagsEditor/Source/GameplayTagsEditor/Private/SGameplayTagQueryWidget.cpp
@@ -261,7 +261,11 @@ TWeakPtr<SGameplayTagQueryWidget> OpenGameplayTagQueryWindow(const FGameplayTagQ
 		];
 
 	// NOTE: FGlobalTabmanager::Get()-> is actually dereferencing a SharedReference, not a SharedPtr, so it cannot be null.
-	if (FGlobalTabmanager::Get()->GetRootWindow().IsValid())
+	if (FSlateApplication::Get().GetActiveModalWindow().IsValid())
+	{
+		FSlateApplication::Get().AddWindowAsNativeChild(Window, FSlateApplication::Get().GetActiveModalWindow().ToSharedRef());
+	}
+	else if (FGlobalTabmanager::Get()->GetRootWindow().IsValid())
 	{
 		FSlateApplication::Get().AddWindowAsNativeChild(Window, FGlobalTabmanager::Get()->GetRootWindow().ToSharedRef());
 	}
