Adding a new GameplayTags function to move tags between ini files

diff --git a/Engine/Plugins/Editor/GameplayTagsEditor/Source/GameplayTagsEditor/Private/GameplayTagsEditorModule.cpp b/Engine/Plugins/Editor/GameplayTagsEditor/Source/GameplayTagsEditor/Private/GameplayTagsEditorModule.cpp
--- a/Engine/Plugins/Editor/GameplayTagsEditor/Source/GameplayTagsEditor/Private/GameplayTagsEditorModule.cpp
+++ b/Engine/Plugins/Editor/GameplayTagsEditor/Source/GameplayTagsEditor/Private/GameplayTagsEditorModule.cpp
@@ -779,6 +779,124 @@ public:
 		return true;
 	}
 
+	virtual bool MoveTagsBetweenINI(const TArray<FString>& TagsToMove, const FName& TargetTagSource, TArray<FString>& OutTagsMoved, TArray<FString>& OutFailedToMoveTags) override
+	{
+		UGameplayTagsManager& Manager = UGameplayTagsManager::Get();
+
+		// Find and check out the destination .ini file
+		FGameplayTagSource* NewTagSource = Manager.FindTagSource(TargetTagSource);
+		if (NewTagSource == nullptr)
+		{
+			ShowNotification(FText::Format(LOCTEXT("MoveTagsFailure_UnknownTarget", "Failed to move tags as target {0} could not be found"), FText::FromName(TargetTagSource)), 10.0f, true);
+			return false;
+		}
+
+		if (NewTagSource->SourceType != EGameplayTagSourceType::DefaultTagList && NewTagSource->SourceType != EGameplayTagSourceType::TagList)
+		{
+			ShowNotification(FText::Format(LOCTEXT("MoveTagsFailure_UnsupportedTarget", "Invalid target source `{0}`! Tags can only be moved to DefaultTagList and TagList target sources."), FText::FromName(TargetTagSource)), 10.0f, true);
+			return false;
+		}
+
+		ensure(NewTagSource->SourceTagList);
+		
+		// Tracking which lists are modified for bulk operations (checkout, config file update/reload).
+		TSet<UGameplayTagsList*> ModifiedTagsList;
+
+		// For each gameplay tag, remove it from the current GameplayTagList and add it to the destination
+		for (const FString& TagToMove : TagsToMove)
+		{
+			FString Comment;
+			TArray<FName> TagSourceNames;
+			bool bIsTagExplicit, bIsRestrictedTag, bAllowNonRestrictedChildren;
+			Manager.GetTagEditorData(FName(TagToMove), Comment, TagSourceNames, bIsTagExplicit, bIsRestrictedTag, bAllowNonRestrictedChildren);
+
+			if (bIsRestrictedTag)
+			{
+				ShowNotification(FText::Format(LOCTEXT("MoveTagsFailure_RestrictedTag", "Restriced Tag {0} cannot be moved"), FText::FromString(TagToMove)), 10.0f, true);
+				OutFailedToMoveTags.Add(TagToMove);
+				continue;
+			}
+
+			if (TagSourceNames.IsEmpty())
+			{
+				ShowNotification(FText::Format(LOCTEXT("MoveTagsFailure_UnknownSource", "Tag {0} sources could not be found"), FText::FromString(TagToMove)), 10.0f, true);
+				OutFailedToMoveTags.Add(TagToMove);
+				continue;
+			}
+			else if (TagSourceNames.Num() > 1)
+			{
+				UE_LOG(LogGameplayTags, Display, TEXT("%d tag sources found for tag %s. Moving first found ini source only! (DefaultTagList or TagList)"), TagSourceNames.Num(), *TagToMove);
+			}
+
+			// Move the first found .ini tag list only
+			FGameplayTagSource* OldTagSource = nullptr;
+			for (FName TagSourceName : TagSourceNames)
+			{
+				FGameplayTagSource* TagSource = Manager.FindTagSource(TagSourceName);
+				if (TagSource != nullptr && (TagSource->SourceType == EGameplayTagSourceType::DefaultTagList || TagSource->SourceType == EGameplayTagSourceType::TagList))
+				{
+					OldTagSource = TagSource;
+					break;
+				}
+			}
+
+			if (OldTagSource == nullptr)
+			{
+				ShowNotification(FText::Format(LOCTEXT("MoveTagsFailure_InvalidSource", "Invalid source for `{0}`! Tags can only be moved from DefaultTagList and TagList sources."), FText::FromString(TagToMove)), 10.0f, true);
+				OutFailedToMoveTags.Add(TagToMove);
+				continue;
+			}
+
+			ensure(OldTagSource->SourceTagList);
+
+			bool bFound = false;
+			FName TagToMoveName(TagToMove);
+
+			// Remove from the old tag source
+			TArray<FGameplayTagTableRow>& GameplayTagList = OldTagSource->SourceTagList->GameplayTagList;
+			for (int32 Index = 0; Index < GameplayTagList.Num(); ++Index)
+			{
+				if (GameplayTagList[Index].Tag == TagToMoveName)
+				{
+					GameplayTagList.RemoveAt(Index);
+					ModifiedTagsList.Add(OldTagSource->SourceTagList);
+					bFound = true;
+					break;
+				}
+			}
+
+			if (bFound)
+			{
+				// Add to the new tag source
+				NewTagSource->SourceTagList->GameplayTagList.AddUnique(FGameplayTagTableRow(FName(TagToMove), Comment));
+				ModifiedTagsList.Add(NewTagSource->SourceTagList);
+
+				OutTagsMoved.Add(TagToMove);
+			}
+			else
+			{
+				ShowNotification(FText::Format(LOCTEXT("MoveTagsFailure_UnknownSource", "Tag {0} could not be found in the source tag list {0}"), FText::FromString(TagToMove), FText::FromString(OldTagSource->SourceTagList->ConfigFileName)), 10.0f, true);
+				OutFailedToMoveTags.Add(TagToMove);
+			}
+		}
+
+		// Update all modified tags list
+		for (UGameplayTagsList* TagsList : ModifiedTagsList)
+		{
+			GameplayTagsUpdateSourceControl(TagsList->ConfigFileName);
+
+			TagsList->SortTags();
+			TagsList->TryUpdateDefaultConfigFile(TagsList->ConfigFileName);
+
+			GConfig->LoadFile(TagsList->ConfigFileName);
+		};
+
+		// Refresh editor
+		Manager.EditorRefreshGameplayTagTree();
+
+		return true;
+	}
+
 	virtual bool AddTransientEditorGameplayTag(const FString& NewTransientTag) override
 	{
 		UGameplayTagsManager& Manager = UGameplayTagsManager::Get();
diff --git a/Engine/Plugins/Editor/GameplayTagsEditor/Source/GameplayTagsEditor/Public/GameplayTagsEditorModule.h b/Engine/Plugins/Editor/GameplayTagsEditor/Source/GameplayTagsEditor/Public/GameplayTagsEditorModule.h
--- a/Engine/Plugins/Editor/GameplayTagsEditor/Source/GameplayTagsEditor/Public/GameplayTagsEditorModule.h
+++ b/Engine/Plugins/Editor/GameplayTagsEditor/Source/GameplayTagsEditor/Public/GameplayTagsEditorModule.h
@@ -56,6 +56,9 @@ public:
 	/** Updates info about a tag */
 	GAMEPLAYTAGSEDITOR_API virtual bool UpdateTagInINI(const FString& TagToUpdate, const FString& Comment, bool bIsRestrictedTag, bool bAllowNonRestrictedChildren) = 0;
 
+	/** Tries to move existing tags from each source ini lists to the target source ini list. */
+	GAMEPLAYTAGSEDITOR_API virtual bool MoveTagsBetweenINI(const TArray<FString>& TagsToMove, const FName& TargetTagSource, TArray<FString>& OutTagsMoved, TArray<FString>& OutFailedToMoveTags) = 0;
+
 	/** Adds a transient gameplay tag (only valid for the current editor session) */
 	GAMEPLAYTAGSEDITOR_API virtual bool AddTransientEditorGameplayTag(const FString& NewTransientTag) = 0;
 
