GameplayTags: Export editor slate classes

diff --git a/Engine/Plugins/Editor/GameplayTagsEditor/Source/GameplayTagsEditor/Public/SGameplayTagCombo.h b/Engine/Plugins/Editor/GameplayTagsEditor/Source/GameplayTagsEditor/Public/SGameplayTagCombo.h
--- a/Engine/Plugins/Editor/GameplayTagsEditor/Source/GameplayTagsEditor/Public/SGameplayTagCombo.h
+++ b/Engine/Plugins/Editor/GameplayTagsEditor/Source/GameplayTagsEditor/Public/SGameplayTagCombo.h
@@ -20,7 +20,7 @@ class SGameplayTagPicker;
  */
 class SGameplayTagCombo : public SCompoundWidget
 {
-	SLATE_DECLARE_WIDGET(SGameplayTagCombo, SCompoundWidget)
+	SLATE_DECLARE_WIDGET_API(SGameplayTagCombo, SCompoundWidget, GAMEPLAYTAGSEDITOR_API)
 
 public:
 
@@ -64,10 +64,10 @@ public:
 
 	GAMEPLAYTAGSEDITOR_API void Construct(const FArguments& InArgs);
 
+	GAMEPLAYTAGSEDITOR_API virtual void Tick(const FGeometry& AllottedGeometry, const double InCurrentTime, const float InDeltaTime) override;
+
 private:
 
-	virtual void Tick(const FGeometry& AllottedGeometry, const double InCurrentTime, const float InDeltaTime) override;
-	
 	bool ShowClearButton() const;
 	FText GetText() const;
 	bool IsValueEnabled() const;
diff --git a/Engine/Plugins/Editor/GameplayTagsEditor/Source/GameplayTagsEditor/Public/SGameplayTagContainerCombo.h b/Engine/Plugins/Editor/GameplayTagsEditor/Source/GameplayTagsEditor/Public/SGameplayTagContainerCombo.h
--- a/Engine/Plugins/Editor/GameplayTagsEditor/Source/GameplayTagsEditor/Public/SGameplayTagContainerCombo.h
+++ b/Engine/Plugins/Editor/GameplayTagsEditor/Source/GameplayTagsEditor/Public/SGameplayTagContainerCombo.h
@@ -22,7 +22,7 @@ class SComboButton;
  */
 class SGameplayTagContainerCombo : public SCompoundWidget
 {
-	SLATE_DECLARE_WIDGET(SGameplayTagContainerCombo, SCompoundWidget)
+	SLATE_DECLARE_WIDGET_API(SGameplayTagContainerCombo, SCompoundWidget, GAMEPLAYTAGSEDITOR_API)
 	
 public:
 
@@ -66,10 +66,10 @@ public:
 
 	GAMEPLAYTAGSEDITOR_API void Construct(const FArguments& InArgs);
 
+	GAMEPLAYTAGSEDITOR_API virtual void Tick(const FGeometry& AllottedGeometry, const double InCurrentTime, const float InDeltaTime) override;
+
 private:
 
-	virtual void Tick(const FGeometry& AllottedGeometry, const double InCurrentTime, const float InDeltaTime) override;
-	
 	struct FEditableItem
 	{
 		FEditableItem() = default;
diff --git a/Engine/Plugins/Editor/GameplayTagsEditor/Source/GameplayTagsEditor/Public/SGameplayTagQueryEntryBox.h b/Engine/Plugins/Editor/GameplayTagsEditor/Source/GameplayTagsEditor/Public/SGameplayTagQueryEntryBox.h
--- a/Engine/Plugins/Editor/GameplayTagsEditor/Source/GameplayTagsEditor/Public/SGameplayTagQueryEntryBox.h
+++ b/Engine/Plugins/Editor/GameplayTagsEditor/Source/GameplayTagsEditor/Public/SGameplayTagQueryEntryBox.h
@@ -16,7 +16,7 @@ class SGameplayTagQueryWidget;
  */
 class SGameplayTagQueryEntryBox : public SCompoundWidget
 {
-	SLATE_DECLARE_WIDGET(SGameplayTagQueryEntryBox, SCompoundWidget)
+	SLATE_DECLARE_WIDGET_API(SGameplayTagQueryEntryBox, SCompoundWidget, GAMEPLAYTAGSEDITOR_API)
 	
 public:
 
@@ -50,7 +50,7 @@ public:
 
 	GAMEPLAYTAGSEDITOR_API void Construct(const FArguments& InArgs);
 
-	virtual void Tick(const FGeometry& AllottedGeometry, const double InCurrentTime, const float InDeltaTime) override;
+	GAMEPLAYTAGSEDITOR_API virtual void Tick(const FGeometry& AllottedGeometry, const double InCurrentTime, const float InDeltaTime) override;
 	
 private:
 
