Trivial fix for LOCTEXT conflict

diff --git a/Engine/Plugins/Editor/GameplayTagsEditor/Source/GameplayTagsEditor/Private/GameplayTagsEditorModule.cpp b/Engine/Plugins/Editor/GameplayTagsEditor/Source/GameplayTagsEditor/Private/GameplayTagsEditorModule.cpp
--- a/Engine/Plugins/Editor/GameplayTagsEditor/Source/GameplayTagsEditor/Private/GameplayTagsEditorModule.cpp
+++ b/Engine/Plugins/Editor/GameplayTagsEditor/Source/GameplayTagsEditor/Private/GameplayTagsEditorModule.cpp
@@ -875,7 +875,7 @@ public:
 			}
 			else
 			{
-				ShowNotification(FText::Format(LOCTEXT("MoveTagsFailure_UnknownSource", "Tag {0} could not be found in the source tag list {0}"), FText::FromString(TagToMove), FText::FromString(OldTagSource->SourceTagList->ConfigFileName)), 10.0f, true);
+				ShowNotification(FText::Format(LOCTEXT("MoveTagsFailure_TagNotFound", "Tag {0} could not be found in the source tag list {0}"), FText::FromString(TagToMove), FText::FromString(OldTagSource->SourceTagList->ConfigFileName)), 10.0f, true);
 				OutFailedToMoveTags.Add(TagToMove);
 			}
 		}
