#UE Add GetUnusedGameplayTags function to GameplayTagsEditorModule. This returns all tags that are in TagLists but not referenced by any assets or ini files

diff --git a/Engine/Plugins/Editor/GameplayTagsEditor/Source/GameplayTagsEditor/Private/GameplayTagsEditorModule.cpp b/Engine/Plugins/Editor/GameplayTagsEditor/Source/GameplayTagsEditor/Private/GameplayTagsEditorModule.cpp
--- a/Engine/Plugins/Editor/GameplayTagsEditor/Source/GameplayTagsEditor/Private/GameplayTagsEditorModule.cpp
+++ b/Engine/Plugins/Editor/GameplayTagsEditor/Source/GameplayTagsEditor/Private/GameplayTagsEditorModule.cpp
@@ -17,6 +17,7 @@
 #include "ISettingsModule.h"
 #include "ISettingsEditorModule.h"
 #include "Misc/CoreDelegates.h"
+#include "Misc/ScopedSlowTask.h"
 #include "Widgets/Notifications/SNotificationList.h"
 #include "Framework/Notifications/NotificationManager.h"
 #include "AssetRegistry/AssetRegistryModule.h"
@@ -1001,6 +1002,142 @@ public:
 			.OnTagChanged(OnChanged);
 	}
 
+	void GetUnusedGameplayTags(TArray<TSharedPtr<FGameplayTagNode>>& OutUnusedTags) override
+	{
+		IAssetRegistry& AssetRegistry = FModuleManager::LoadModuleChecked<FAssetRegistryModule>("AssetRegistry").Get();
+		UGameplayTagsManager& Manager = UGameplayTagsManager::Get();
+		int32 NumUsedExplicitTags = 0;
+		TSet<FString> AllConfigValues;
+
+		// Populate all config values from all config files so that we can check if any config contains a reference to a tag later
+		{
+			TArray<FString> AllConfigFilenames;
+			GConfig->GetConfigFilenames(AllConfigFilenames);
+
+			for (const FString& ConfigFilename : AllConfigFilenames)
+			{
+				if (FConfigFile* ConfigFile = GConfig->FindConfigFile(ConfigFilename))
+				{
+					for (const TPair<FString, FConfigSection>& FileIt : AsConst(*ConfigFile))
+					{
+						const FString& SectionName = FileIt.Key;
+
+						// Do not include sections that define the tags, as we don't wan't their definition showing up as a reference
+						if (SectionName != TEXT("/Script/GameplayTags.GameplayTagsSettings") && SectionName != TEXT("/Script/GameplayTags.GameplayTagsList"))
+						{
+							for (const TPair<FName, FConfigValue>& SectionIt : FileIt.Value)
+							{
+								const FString& ConfigValue = SectionIt.Value.GetValue();
+
+								// Cut down on values by skipping pure numbers
+								if (!ConfigValue.IsNumeric())
+								{
+									AllConfigValues.Add(ConfigValue);
+								}
+							}
+						}
+					}
+				}
+			}
+		}
+
+		FScopedSlowTask SlowTask((float)Manager.GetNumGameplayTagNodes(), LOCTEXT("PopulatingUnusedTags", "Populating Unused Tags"));
+
+		// Function to determine if a single node is referenced by content
+		auto IsNodeUsed = [&AssetRegistry, &AllConfigValues](const TSharedPtr<FGameplayTagNode>& Node) -> bool
+		{
+			// Look for asset references
+			FAssetIdentifier TagId = FAssetIdentifier(FGameplayTag::StaticStruct(), Node->GetCompleteTagName());
+			TArray<FAssetIdentifier> Referencers;
+			AssetRegistry.GetReferencers(TagId, Referencers, UE::AssetRegistry::EDependencyCategory::SearchableName);
+			if (Referencers.Num() != 0)
+			{
+				return true;
+			}
+
+			// Look for config references
+			FString TagString = Node->GetCompleteTagString();
+			for (const FString& ConfigValue : AllConfigValues)
+			{
+				if (ConfigValue.Contains(TagString))
+				{
+					return true;
+				}
+			}
+
+			return false;
+		};
+
+		// Recursive function to traverse the Gameplay Tag Node tree and find all unused tags
+		TFunction<void(const TSharedPtr<FGameplayTagNode>&)> RecursiveProcessTagNode;
+		RecursiveProcessTagNode = [&RecursiveProcessTagNode, &NumUsedExplicitTags, &Manager, &SlowTask, &IsNodeUsed, &OutUnusedTags](const TSharedPtr<FGameplayTagNode>& Node)
+		{
+			check(Node.IsValid());
+
+			SlowTask.EnterProgressFrame();
+
+			if (Node->IsExplicitTag())
+			{
+				bool bSourcesDetectable = false;
+				for (const FName& SourceName : Node->GetAllSourceNames())
+				{
+					const FGameplayTagSource* TagSource = Manager.FindTagSource(SourceName);
+					if (TagSource &&
+						(TagSource->SourceType == EGameplayTagSourceType::DefaultTagList || TagSource->SourceType == EGameplayTagSourceType::TagList || TagSource->SourceType == EGameplayTagSourceType::RestrictedTagList))
+					{
+						bSourcesDetectable = true;
+					}
+					else
+					{
+						bSourcesDetectable = false;
+						break;
+					}
+				}
+
+				if (bSourcesDetectable && !IsNodeUsed(Node))
+				{
+					OutUnusedTags.Add(Node);
+				}
+				else
+				{
+					NumUsedExplicitTags++;
+				}
+			}
+
+			// Iterate children recursively
+			const int32 NumUsedExplicitTagsBeforeChildren = NumUsedExplicitTags;
+			const int32 NumUnusedExplicitTagsBeforeChildren = OutUnusedTags.Num();
+			const TArray<TSharedPtr<FGameplayTagNode>>& ChildNodes = Node->GetChildTagNodes();
+			for (const TSharedPtr<FGameplayTagNode>& ChildNode : ChildNodes)
+			{
+				RecursiveProcessTagNode(ChildNode);
+			}
+
+			// Implicit tags need at least one explicit child in order to exist.
+			// If an implicit tag is referenced, treat the last explicit child as being referenced too
+			const bool bAtLeastOneUsedExplicitChild = NumUsedExplicitTags > NumUsedExplicitTagsBeforeChildren;
+			const bool bAtLeastOneUnusedExplicitChild = OutUnusedTags.Num() > NumUnusedExplicitTagsBeforeChildren;
+			if (!Node->IsExplicitTag() && !bAtLeastOneUsedExplicitChild && bAtLeastOneUnusedExplicitChild)
+			{
+				// This is an implicit tag that has only unused explicit children. If this tag is referenced, then remove the last unused child from the list
+				if (IsNodeUsed(Node))
+				{
+					// This implicit tag is referenced. Remove the last unused child as this tag is using that child in order to exist.
+					OutUnusedTags.RemoveAt(OutUnusedTags.Num() - 1, 1, EAllowShrinking::No);
+					NumUsedExplicitTags++;
+				}
+			}
+		};
+
+		// Go through all root nodes to process entire tree
+		TArray<TSharedPtr<FGameplayTagNode>> AllRoots;
+		Manager.GetFilteredGameplayRootTags(TEXT(""), AllRoots);
+		for (const TSharedPtr<FGameplayTagNode>& Root : AllRoots)
+		{
+			RecursiveProcessTagNode(Root);
+		}
+	}
+
 	static bool WriteCustomReport(FString FileName, TArray<FString>& FileLines)
 	{
 		// Has a report been generated
diff --git a/Engine/Plugins/Editor/GameplayTagsEditor/Source/GameplayTagsEditor/Public/GameplayTagsEditorModule.h b/Engine/Plugins/Editor/GameplayTagsEditor/Source/GameplayTagsEditor/Public/GameplayTagsEditorModule.h
--- a/Engine/Plugins/Editor/GameplayTagsEditor/Source/GameplayTagsEditor/Public/GameplayTagsEditorModule.h
+++ b/Engine/Plugins/Editor/GameplayTagsEditor/Source/GameplayTagsEditor/Public/GameplayTagsEditorModule.h
@@ -80,6 +80,9 @@ public:
 	 * @param FilterString		Optional filter string, same format as Categories metadata on tag properties
 	 */
 	GAMEPLAYTAGSEDITOR_API virtual TSharedRef<SWidget> MakeGameplayTagWidget(FOnSetGameplayTag OnSetTag, TSharedPtr<FGameplayTag> GameplayTag, const FString& FilterString = FString()) = 0;
+
+	/** Returns the list of gameplay tags that are not used by content */
+	GAMEPLAYTAGSEDITOR_API virtual void GetUnusedGameplayTags(TArray<TSharedPtr<struct FGameplayTagNode>>& OutUnusedTags) = 0;
 };
 
 /** This is public so that child structs of FGameplayTag can use the details customization */
