#UE GameplayTagsEditor tag deletion improvements.

* Add DeleteTagsFromINI() function to allow deleting multiple tags at once and batch all expensive operations that can be done faster in bulk
* DeleteTagFromINI now properly considers all tag sources instead of only the first one
* Fixed bug where parent implicit tags were not actually being checked for references and thus allowing explicit tags to be removed such that the referenced implicit parents were also removed.
* Fixed bug that happened when a tag was listed twice in a source file. Only the first entry was removed.
* Fixed bug where ini files that were not checked out before deleting were not being updated in memory after the delete.
* Add option in ShowNotification() to only emit the message to the log. Useful when doing bulk operations that may emit many notifications so it doesn't spam the log

diff --git a/Engine/Plugins/Editor/GameplayTagsEditor/Source/GameplayTagsEditor/Private/GameplayTagsEditorModule.cpp b/Engine/Plugins/Editor/GameplayTagsEditor/Source/GameplayTagsEditor/Private/GameplayTagsEditorModule.cpp
--- a/Engine/Plugins/Editor/GameplayTagsEditor/Source/GameplayTagsEditor/Private/GameplayTagsEditorModule.cpp
+++ b/Engine/Plugins/Editor/GameplayTagsEditor/Source/GameplayTagsEditor/Private/GameplayTagsEditorModule.cpp
@@ -191,17 +191,31 @@ public:
 		}
 	}
 
-	void ShowNotification(const FText& TextToDisplay, float TimeToDisplay, bool bLogError = false)
+	void ShowNotification(const FText& TextToDisplay, float TimeToDisplay, bool bLogError = false, bool bOnlyLog = false)
 	{
-		FNotificationInfo Info(TextToDisplay);
-		Info.ExpireDuration = TimeToDisplay;
+		if (bOnlyLog)
+		{
+			if (bLogError)
+			{
+				UE_LOG(LogGameplayTags, Error, TEXT("%s"), *TextToDisplay.ToString())
+			}
+			else
+			{
+				UE_LOG(LogGameplayTags, Display, TEXT("%s"), *TextToDisplay.ToString())
+			}
+		}
+		else
+		{
+			FNotificationInfo Info(TextToDisplay);
+			Info.ExpireDuration = TimeToDisplay;
 
-		FSlateNotificationManager::Get().AddNotification(Info);
+			FSlateNotificationManager::Get().AddNotification(Info);
 
-		// Also log if error
-		if (bLogError)
-		{
-			UE_LOG(LogGameplayTags, Error, TEXT("%s"), *TextToDisplay.ToString())
+			// Also log if error
+			if (bLogError)
+			{
+				UE_LOG(LogGameplayTags, Error, TEXT("%s"), *TextToDisplay.ToString())
+			}
 		}
 	}
 
@@ -263,34 +277,55 @@ public:
 		ShowNotification(LOCTEXT("MigrationText", "Migrated Tag Settings, check DefaultEngine.ini before checking in!"), 10.0f);
 	}
 
-	void GameplayTagsUpdateSourceControl(const FString& RelativeConfigFilePath)
+	void GameplayTagsUpdateSourceControl(const FString& RelativeConfigFilePath, bool bOnlyLog = false)
 	{
-		FString ConfigPath = FPaths::ConvertRelativePathToFull(RelativeConfigFilePath);
+		TArray<FString> RelativeConfigFilePaths = { RelativeConfigFilePath };
+		GameplayTagsUpdateSourceControl(RelativeConfigFilePaths, bOnlyLog);
+	}
 
-		if (!FPlatformFileManager::Get().GetPlatformFile().FileExists(*ConfigPath))
+	void GameplayTagsUpdateSourceControl(const TArray<FString>& RelativeConfigFilePaths, bool bOnlyLog = false)
+	{
+		TArray<FString> ExistingConfigPaths;
+		for (const FString& RelativeConfigFilePath : RelativeConfigFilePaths)
 		{
-			return;
+			FString ConfigPath = FPaths::ConvertRelativePathToFull(RelativeConfigFilePath);
+
+			if (FPlatformFileManager::Get().GetPlatformFile().FileExists(*ConfigPath))
+			{
+				ExistingConfigPaths.Add(ConfigPath);
+			}
 		}
 
 		if (ISourceControlModule::Get().IsEnabled())
 		{
 			FText ErrorMessage;
 
-			if (!SourceControlHelpers::CheckoutOrMarkForAdd(ConfigPath, FText::FromString(ConfigPath), NULL, ErrorMessage))
+			if (ExistingConfigPaths.Num() == 1)
+			{
+				const FString& ConfigPath = ExistingConfigPaths[0];
+				if (!SourceControlHelpers::CheckoutOrMarkForAdd(ConfigPath, FText::FromString(ConfigPath), NULL, ErrorMessage))
+				{
+					ShowNotification(ErrorMessage, 3.0f, false, bOnlyLog);
+				}
+			}
+			else
 			{
-				ShowNotification(ErrorMessage, 3.0f);
+				SourceControlHelpers::CheckOutOrAddFiles(ExistingConfigPaths);
 			}
 		}
 		else
 		{
-			if (!FPlatformFileManager::Get().GetPlatformFile().SetReadOnly(*ConfigPath, false))
+			for (const FString& ConfigPath : ExistingConfigPaths)
 			{
-				ShowNotification(FText::Format(LOCTEXT("FailedToMakeWritable", "Could not make {0} writable."), FText::FromString(ConfigPath)), 3.0f);
+				if (!FPlatformFileManager::Get().GetPlatformFile().SetReadOnly(*ConfigPath, false))
+				{
+					ShowNotification(FText::Format(LOCTEXT("FailedToMakeWritable", "Could not make {0} writable."), FText::FromString(ConfigPath)), 3.0f, false, bOnlyLog);
+				}
 			}
 		}
 	}
 
-	bool DeleteTagRedirector(const FName& TagToDelete)
+	bool DeleteTagRedirector(const FName& TagToDelete, bool bOnlyLog = false, bool bRefresh = true, TMap<UObject*, FString>* OutObjectsToUpdateConfig = nullptr)
 	{
 		UGameplayTagsSettings* Settings = GetMutableDefault<UGameplayTagsSettings>();
 		UGameplayTagsManager& Manager = UGameplayTagsManager::Get();
@@ -301,19 +336,32 @@ public:
 			{
 				Settings->GameplayTagRedirects.RemoveAt(i);
 
-				GameplayTagsUpdateSourceControl(Settings->GetDefaultConfigFilename());
-				Settings->TryUpdateDefaultConfigFile();
-				GConfig->LoadFile(Settings->GetDefaultConfigFilename());
+				if (bRefresh)
+				{
+					GameplayTagsUpdateSourceControl(Settings->GetDefaultConfigFilename());
+					Settings->TryUpdateDefaultConfigFile();
+					GConfig->LoadFile(Settings->GetDefaultConfigFilename());
 
-				Manager.EditorRefreshGameplayTagTree();
+					Manager.EditorRefreshGameplayTagTree();
+				}
+				else
+				{
+					if (OutObjectsToUpdateConfig)
+					{
+						OutObjectsToUpdateConfig->Add(Settings, Settings->GetDefaultConfigFilename());
+					}
+				}
 
-				ShowNotification(FText::Format(LOCTEXT("RemoveTagRedirect", "Deleted tag redirect {0}"), FText::FromName(TagToDelete)), 5.0f);
+				ShowNotification(FText::Format(LOCTEXT("RemoveTagRedirect", "Deleted tag redirect {0}"), FText::FromName(TagToDelete)), 5.0f, false, bOnlyLog);
 
 				WarnAboutRestart();
 
-				TSharedPtr<FGameplayTagNode> FoundNode = Manager.FindTagNode(TagToDelete);
+				if (bRefresh)
+				{
+					TSharedPtr<FGameplayTagNode> FoundNode = Manager.FindTagNode(TagToDelete);
 
-				ensureMsgf(!FoundNode.IsValid() || FoundNode->GetCompleteTagName() == TagToDelete, TEXT("Failed to delete redirector %s!"), *TagToDelete.ToString());
+					ensureMsgf(!FoundNode.IsValid() || FoundNode->GetCompleteTagName() == TagToDelete, TEXT("Failed to delete redirector %s!"), *TagToDelete.ToString());
+				}
 
 				return true;
 			}
@@ -493,6 +541,56 @@ public:
 	}
 
 	virtual bool DeleteTagFromINI(TSharedPtr<FGameplayTagNode> TagNodeToDelete) override
+	{
+		if (!TagNodeToDelete.IsValid())
+		{
+			return false;
+		}
+
+		TMap<UObject*, FString> ObjectsToUpdateConfig;
+		const bool bOnlyLog = false;
+		bool bReturnValue = DeleteTagFromINIInternal(TagNodeToDelete, bOnlyLog, ObjectsToUpdateConfig);
+		if (ObjectsToUpdateConfig.Num() > 0)
+		{
+			UpdateTagSourcesAfterDelete(bOnlyLog, ObjectsToUpdateConfig);
+
+			// This invalidates all local variables, need to return right away
+			UGameplayTagsManager::Get().EditorRefreshGameplayTagTree();
+		}
+		return bReturnValue;
+	}
+
+	virtual void DeleteTagsFromINI(const TArray<TSharedPtr<FGameplayTagNode>>& TagNodesToDelete) override
+	{
+		TMap<UObject*, FString> ObjectsToUpdateConfig;
+		const bool bOnlyLog = true;
+
+		{
+			FScopedSlowTask SlowTask(TagNodesToDelete.Num(), LOCTEXT("RemovingTags", "Removing Tags"));
+			for (const TSharedPtr<FGameplayTagNode>& TagNodeToDelete : TagNodesToDelete)
+			{
+				if (ensure(!TagNodeToDelete->GetCompleteTagName().IsNone()))
+				{
+					SlowTask.EnterProgressFrame();
+					if (TagNodeToDelete.IsValid())
+					{
+						DeleteTagFromINIInternal(TagNodeToDelete, bOnlyLog, ObjectsToUpdateConfig);
+					}
+					ensureMsgf(!TagNodeToDelete->GetCompleteTagName().IsNone(), TEXT("A 'None' tag here implies somone may have added a EditorRefreshGameplayTagTree() call in DeleteTagFromINI. Do not do this, the refresh must happen after the bulk operation is done."));
+				}
+			}
+		}
+
+		if (ObjectsToUpdateConfig.Num() > 0)
+		{
+			UpdateTagSourcesAfterDelete(bOnlyLog, ObjectsToUpdateConfig);
+
+			UGameplayTagsManager& Manager = UGameplayTagsManager::Get();
+			Manager.EditorRefreshGameplayTagTree();
+		}
+	}
+
+	bool DeleteTagFromINIInternal(const TSharedPtr<FGameplayTagNode>& TagNodeToDelete, bool bOnlyLog, TMap<UObject*, FString>& OutObjectsToUpdateConfig)
 	{
 		FName TagName = TagNodeToDelete->GetCompleteTagName();
 
@@ -500,42 +598,28 @@ public:
 		UGameplayTagsSettings* Settings = GetMutableDefault<UGameplayTagsSettings>();
 
 		FString Comment;
-		FName TagSourceName;
+		TArray<FName> TagSourceNames;
 		bool bTagIsExplicit;
 		bool bTagIsRestricted;
 		bool bTagAllowsNonRestrictedChildren;
 
-		if (DeleteTagRedirector(TagName))
+		if (DeleteTagRedirector(TagName, bOnlyLog, false, &OutObjectsToUpdateConfig))
 		{
 			return true;
 		}
 		
-		if (!Manager.GetTagEditorData(TagName, Comment, TagSourceName, bTagIsExplicit, bTagIsRestricted, bTagAllowsNonRestrictedChildren))
+		if (!Manager.GetTagEditorData(TagName, Comment, TagSourceNames, bTagIsExplicit, bTagIsRestricted, bTagAllowsNonRestrictedChildren))
 		{
-			ShowNotification(FText::Format(LOCTEXT("RemoveTagFailureNoTag", "Cannot delete tag {0}, does not exist!"), FText::FromName(TagName)), 10.0f, true);
+			ShowNotification(FText::Format(LOCTEXT("RemoveTagFailureNoTag", "Cannot delete tag {0}, does not exist!"), FText::FromName(TagName)), 10.0f, true, bOnlyLog);
 			return false;
 		}
 
 		ensure(bTagIsRestricted == TagNodeToDelete->IsRestrictedGameplayTag());
 
-		const FGameplayTagSource* TagSource = Manager.FindTagSource(TagSourceName);
-
 		// Check if the tag is implicitly defined
-		if (!bTagIsExplicit || !TagSource)
+		if (!bTagIsExplicit || TagSourceNames.Num() == 0)
 		{
-			ShowNotification(FText::Format(LOCTEXT("RemoveTagFailureNoSource", "Cannot delete tag {0} as it is implicit, remove children manually"), FText::FromName(TagName)), 10.0f, true);
-			return false;
-		}
-		
-		if (bTagIsRestricted && !TagSource->SourceRestrictedTagList)
-		{
-			ShowNotification(FText::Format(LOCTEXT("RemoveTagFailureBadSource", "Cannot delete tag {0} from source {1}, remove manually"), FText::FromName(TagName), FText::FromName(TagSourceName)), 10.0f, true);
-			return false;
-		}
-
-		if (!bTagIsRestricted && !TagSource->SourceTagList)
-		{
-			ShowNotification(FText::Format(LOCTEXT("RemoveTagFailureBadSource", "Cannot delete tag {0} from source {1}, remove manually"), FText::FromName(TagName), FText::FromName(TagSourceName)), 10.0f, true);
+			ShowNotification(FText::Format(LOCTEXT("RemoveTagFailureNoSource", "Cannot delete tag {0} as it is implicit, remove children manually"), FText::FromName(TagName)), 10.0f, true, bOnlyLog);
 			return false;
 		}
 
@@ -568,7 +652,7 @@ public:
 		for (FName TagNameToDelete : TagsThatWillBeDeleted)
 		{
 			// Verify references
-			FAssetIdentifier TagId = FAssetIdentifier(FGameplayTag::StaticStruct(), TagName);
+			FAssetIdentifier TagId = FAssetIdentifier(FGameplayTag::StaticStruct(), TagNameToDelete);
 			TArray<FAssetIdentifier> Referencers;
 
 			FAssetRegistryModule& AssetRegistryModule = FModuleManager::LoadModuleChecked<FAssetRegistryModule>("AssetRegistry");
@@ -576,64 +660,98 @@ public:
 
 			if (Referencers.Num() > 0)
 			{
-				ShowNotification(FText::Format(LOCTEXT("RemoveTagFailureBadSource_Referenced", "Cannot delete tag {0}, still referenced by {1} and possibly others"), FText::FromName(TagNameToDelete), FText::FromString(Referencers[0].ToString())), 10.0f, true);
+				ShowNotification(FText::Format(LOCTEXT("RemoveTagFailureBadSource_Referenced", "Cannot delete tag {0}, still referenced by {1} and possibly others"), FText::FromName(TagNameToDelete), FText::FromString(Referencers[0].ToString())), 10.0f, true, bOnlyLog);
 
 				return false;
 			}
 		}
 
-		// Passed, delete and save
-		const FString& ConfigFileName = bTagIsRestricted ? TagSource->SourceRestrictedTagList->ConfigFileName : TagSource->SourceTagList->ConfigFileName;
-		int32 TagListSize = bTagIsRestricted ? TagSource->SourceRestrictedTagList->RestrictedGameplayTagList.Num() : TagSource->SourceTagList->GameplayTagList.Num();
-
-		for (int32 i = 0; i < TagListSize; i++)
+		bool bRemovedAny = false;
+		for (const FName& TagSourceName : TagSourceNames)
 		{
-			bool bRemoved = false;
+			const FGameplayTagSource* TagSource = Manager.FindTagSource(TagSourceName);
+
+			if (bTagIsRestricted && !TagSource->SourceRestrictedTagList)
+			{
+				ShowNotification(FText::Format(LOCTEXT("RemoveTagFailureBadSource", "Cannot delete tag {0} from source {1}, remove manually"), FText::FromName(TagName), FText::FromName(TagSourceName)), 10.0f, true, bOnlyLog);
+				continue;
+			}
+
+			if (!bTagIsRestricted && !TagSource->SourceTagList)
+			{
+				ShowNotification(FText::Format(LOCTEXT("RemoveTagFailureBadSource", "Cannot delete tag {0} from source {1}, remove manually"), FText::FromName(TagName), FText::FromName(TagSourceName)), 10.0f, true, bOnlyLog);
+				continue;
+			}
+
+			// Passed, delete and save
+			const FString& ConfigFileName = bTagIsRestricted ? TagSource->SourceRestrictedTagList->ConfigFileName : TagSource->SourceTagList->ConfigFileName;
+
+			int32 NumRemoved = 0;
 			if (bTagIsRestricted)
 			{
-				if (TagSource->SourceRestrictedTagList->RestrictedGameplayTagList[i].Tag == TagName)
+				NumRemoved = TagSource->SourceRestrictedTagList->RestrictedGameplayTagList.RemoveAll([&TagName](const FRestrictedGameplayTagTableRow& Row) { return Row.Tag == TagName; });
+				if (NumRemoved > 0)
 				{
-					TagSource->SourceRestrictedTagList->RestrictedGameplayTagList.RemoveAt(i);
-					TagSource->SourceRestrictedTagList->TryUpdateDefaultConfigFile(ConfigFileName);
-					bRemoved = true;
+					OutObjectsToUpdateConfig.Add(TagSource->SourceRestrictedTagList, ConfigFileName);
 				}
 			}
 			else
 			{
-				if (TagSource->SourceTagList->GameplayTagList[i].Tag == TagName)
+				NumRemoved = TagSource->SourceTagList->GameplayTagList.RemoveAll([&TagName](const FGameplayTagTableRow& Row) { return Row.Tag == TagName; });
+				if (NumRemoved > 0)
 				{
-					TagSource->SourceTagList->GameplayTagList.RemoveAt(i);
-					TagSource->SourceTagList->TryUpdateDefaultConfigFile(ConfigFileName);
-					bRemoved = true;
+					OutObjectsToUpdateConfig.Add(TagSource->SourceTagList, ConfigFileName);
 				}
 			}
 
-			if (bRemoved)
+			if (NumRemoved > 0)
 			{
-				GameplayTagsUpdateSourceControl(ConfigFileName);
-				GConfig->LoadFile(ConfigFileName);
-
 				// See if we still live due to child tags
-
 				if (ChildTags.Num() > 0)
 				{
-					ShowNotification(FText::Format(LOCTEXT("RemoveTagChildrenExist", "Deleted explicit tag {0}, still exists implicitly due to children"), FText::FromName(TagName)), 5.0f);
+					ShowNotification(FText::Format(LOCTEXT("RemoveTagChildrenExist", "Deleted explicit tag {0}, still exists implicitly due to children"), FText::FromName(TagName)), 5.0f, false, bOnlyLog);
 				}
 				else
 				{
-					ShowNotification(FText::Format(LOCTEXT("RemoveTag", "Deleted tag {0}"), FText::FromName(TagName)), 5.0f);
+					ShowNotification(FText::Format(LOCTEXT("RemoveTag", "Deleted tag {0}"), FText::FromName(TagName)), 5.0f, false, bOnlyLog);
 				}
 
-				// This invalidates all local variables, need to return right away
-				Manager.EditorRefreshGameplayTagTree();
-
-				return true;
+				bRemovedAny = true;
 			}
 		}
 
-		ShowNotification(FText::Format(LOCTEXT("RemoveTagFailureNoTag", "Cannot delete tag {0}, does not exist!"), FText::FromName(TagName)), 10.0f, true);
+		if (!bRemovedAny)
+		{
+			ShowNotification(FText::Format(LOCTEXT("RemoveTagFailureNoTag", "Cannot delete tag {0}, does not exist!"), FText::FromName(TagName)), 10.0f, true, bOnlyLog);
+		}
 		
-		return false;
+		return bRemovedAny;
+	}
+
+	void UpdateTagSourcesAfterDelete(bool bOnlyLog, const TMap<UObject*, FString>& ObjectsToUpdateConfig)
+	{
+		TSet<FString> ConfigFileNames;
+		for (auto ObjectToUpdateIt = ObjectsToUpdateConfig.CreateConstIterator(); ObjectToUpdateIt; ++ObjectToUpdateIt)
+		{
+			ConfigFileNames.Add(ObjectToUpdateIt.Value());
+		}
+
+		FScopedSlowTask SlowTask(ObjectsToUpdateConfig.Num() + ConfigFileNames.Num(), LOCTEXT("UpdateTagSourcesAfterDelete", "Updating Tag Sources"));
+
+		GameplayTagsUpdateSourceControl(ConfigFileNames.Array(), bOnlyLog);
+
+		for (auto ObjectToUpdateIt = ObjectsToUpdateConfig.CreateConstIterator(); ObjectToUpdateIt; ++ObjectToUpdateIt)
+		{
+			SlowTask.EnterProgressFrame();
+			check(ObjectToUpdateIt.Key());
+			ObjectToUpdateIt.Key()->TryUpdateDefaultConfigFile(ObjectToUpdateIt.Value());
+		}
+
+		for (const FString& ConfigFileName : ConfigFileNames)
+		{
+			SlowTask.EnterProgressFrame();
+			GConfig->LoadFile(ConfigFileName);
+		}
 	}
 
 	virtual bool UpdateTagInINI(const FString& TagToUpdate, const FString& Comment, bool bIsRestrictedTag, bool bAllowNonRestrictedChildren) override
diff --git a/Engine/Plugins/Editor/GameplayTagsEditor/Source/GameplayTagsEditor/Public/GameplayTagsEditorModule.h b/Engine/Plugins/Editor/GameplayTagsEditor/Source/GameplayTagsEditor/Public/GameplayTagsEditorModule.h
--- a/Engine/Plugins/Editor/GameplayTagsEditor/Source/GameplayTagsEditor/Public/GameplayTagsEditorModule.h
+++ b/Engine/Plugins/Editor/GameplayTagsEditor/Source/GameplayTagsEditor/Public/GameplayTagsEditorModule.h
@@ -49,6 +49,7 @@ public:
 
 	/** Tries to delete a tag from the library. This will pop up special UI or error messages as needed. It will also delete redirectors if that is specified. */
 	GAMEPLAYTAGSEDITOR_API virtual bool DeleteTagFromINI(TSharedPtr<struct FGameplayTagNode> TagNodeToDelete) = 0;
+	GAMEPLAYTAGSEDITOR_API virtual void DeleteTagsFromINI(const TArray<TSharedPtr<struct FGameplayTagNode>>& TagNodesToDelete) = 0;
 
 	/** Tries to rename a tag, leaving a rediretor in the ini, and adding the new tag if it does not exist yet */
 	GAMEPLAYTAGSEDITOR_API virtual bool RenameTagInINI(const FString& TagToRename, const FString& TagToRenameTo) = 0;
