#UE Consider redirector old names when determining if a gameplay tag is unused

diff --git a/Engine/Plugins/Editor/GameplayTagsEditor/Source/GameplayTagsEditor/Private/GameplayTagsEditorModule.cpp b/Engine/Plugins/Editor/GameplayTagsEditor/Source/GameplayTagsEditor/Private/GameplayTagsEditorModule.cpp
--- a/Engine/Plugins/Editor/GameplayTagsEditor/Source/GameplayTagsEditor/Private/GameplayTagsEditorModule.cpp
+++ b/Engine/Plugins/Editor/GameplayTagsEditor/Source/GameplayTagsEditor/Private/GameplayTagsEditorModule.cpp
@@ -1126,6 +1126,7 @@ public:
 		UGameplayTagsManager& Manager = UGameplayTagsManager::Get();
 		int32 NumUsedExplicitTags = 0;
 		TSet<FString> AllConfigValues;
+		TMultiMap<FName, FName> RerverseRedirectorMap;
 
 		// Populate all config values from all config files so that we can check if any config contains a reference to a tag later
 		{
@@ -1159,28 +1160,43 @@ public:
 			}
 		}
 
+		// Build a reverse map of tag redirectors so we can find the old names of a given tag to see if any of them are still referenced
+		UGameplayTagsSettings* Settings = GetMutableDefault<UGameplayTagsSettings>();
+		for (const FGameplayTagRedirect& Redirect : Settings->GameplayTagRedirects)
+		{
+			RerverseRedirectorMap.Add(Redirect.NewTagName, Redirect.OldTagName);
+		}
+
 		FScopedSlowTask SlowTask((float)Manager.GetNumGameplayTagNodes(), LOCTEXT("PopulatingUnusedTags", "Populating Unused Tags"));
 
 		// Function to determine if a single node is referenced by content
-		auto IsNodeUsed = [&AssetRegistry, &AllConfigValues](const TSharedPtr<FGameplayTagNode>& Node) -> bool
+		auto IsNodeUsed = [&AssetRegistry, &AllConfigValues, &RerverseRedirectorMap](const TSharedPtr<FGameplayTagNode>& Node) -> bool
 		{
-			// Look for asset references
-			FAssetIdentifier TagId = FAssetIdentifier(FGameplayTag::StaticStruct(), Node->GetCompleteTagName());
-			TArray<FAssetIdentifier> Referencers;
-			AssetRegistry.GetReferencers(TagId, Referencers, UE::AssetRegistry::EDependencyCategory::SearchableName);
-			if (Referencers.Num() != 0)
-			{
-				return true;
-			}
+			// Look for references to the input tag or any of its old names if there were redirectors
+			FName InitialTagName = Node->GetCompleteTagName();
+			TArray<FName> TagsToCheck = { InitialTagName };
+			RerverseRedirectorMap.MultiFind(InitialTagName, TagsToCheck);
 
-			// Look for config references
-			FString TagString = Node->GetCompleteTagString();
-			for (const FString& ConfigValue : AllConfigValues)
+			for (const FName& TagName : TagsToCheck)
 			{
-				if (ConfigValue.Contains(TagString))
+				// Look for asset references
+				FAssetIdentifier TagId = FAssetIdentifier(FGameplayTag::StaticStruct(), TagName);
+				TArray<FAssetIdentifier> Referencers;
+				AssetRegistry.GetReferencers(TagId, Referencers, UE::AssetRegistry::EDependencyCategory::SearchableName);
+				if (Referencers.Num() != 0)
 				{
 					return true;
 				}
+
+				// Look for config references
+				FString TagString = TagName.ToString();
+				for (const FString& ConfigValue : AllConfigValues)
+				{
+					if (ConfigValue.Contains(TagString))
+					{
+						return true;
+					}
+				}
 			}
 
 			return false;
