#UE Minor UI fix for cleanup unused tags dialog to allow the list to match the size of the window

diff --git a/Engine/Plugins/Editor/GameplayTagsEditor/Source/GameplayTagsEditor/Private/SCleanupUnusedGameplayTagsWidget.cpp b/Engine/Plugins/Editor/GameplayTagsEditor/Source/GameplayTagsEditor/Private/SCleanupUnusedGameplayTagsWidget.cpp
--- a/Engine/Plugins/Editor/GameplayTagsEditor/Source/GameplayTagsEditor/Private/SCleanupUnusedGameplayTagsWidget.cpp
+++ b/Engine/Plugins/Editor/GameplayTagsEditor/Source/GameplayTagsEditor/Private/SCleanupUnusedGameplayTagsWidget.cpp
@@ -212,7 +212,6 @@ void SCleanupUnusedGameplayTagsWidget::Construct(const FArguments& InArgs)
 						]
 
 						+ SVerticalBox::Slot()
-						.MaxHeight(800.f)
 						[
 							SAssignNew(UnusedTagsListView, SListView<TSharedPtr<FUnusedTagItem>>)
 								.ListItemsSource(&UnusedTags)
