Fix the add new gameplay tag widget to also handle adding restricted tags, and delete the broken old UI.
Improve the tag source combo boxes to support searching/auto complete.

diff --git a/Engine/Plugins/Editor/GameplayTagsEditor/Source/GameplayTagsEditor/Internal/SAddNewGameplayTagWidget.h b/Engine/Plugins/Editor/GameplayTagsEditor/Source/GameplayTagsEditor/Internal/SAddNewGameplayTagWidget.h
--- a/Engine/Plugins/Editor/GameplayTagsEditor/Source/GameplayTagsEditor/Internal/SAddNewGameplayTagWidget.h
+++ b/Engine/Plugins/Editor/GameplayTagsEditor/Source/GameplayTagsEditor/Internal/SAddNewGameplayTagWidget.h
@@ -6,7 +6,8 @@
 
 class SEditableTextBox;
 namespace ETextCommit { enum Type : int; }
-template <typename OptionType> class SComboBox;
+class SCheckBox;
+class SSearchableComboBox;
 class SNotificationItem;
 
 /** Widget allowing the user to create new gameplay tags */
@@ -34,6 +35,7 @@ public:
 		SLATE_ARGUMENT(FString, NewTagName) // String that will initially populate the New Tag Name field
 		SLATE_ARGUMENT(FMargin, Padding) // Padding around the widget
 		SLATE_ARGUMENT(FMargin, AddButtonPadding) // padding around the Add button
+		SLATE_ARGUMENT(bool, RestrictedTags)// If we are dealing with restricted tags or regular gameplay tags
 	SLATE_END_ARGS();
 
 	GAMEPLAYTAGSEDITOR_API virtual ~SAddNewGameplayTagWidget() override;
@@ -65,6 +67,9 @@ private:
 	/** Selects tag file location. Uses the default if the location is not specified */
 	void SelectTagSource(const FName& InSource = FName());
 
+	/** Gets selected tag source as name */
+	FName GetSelectedTagSource() const;
+
 	/** Creates a list of all INIs that gameplay tags can be added to */
 	void PopulateTagSources();
 
@@ -74,14 +79,16 @@ private:
 	/** Callback for when the Add New Tag button is pressed */
 	FReply OnAddNewTagButtonPressed();
 
-	/** Return true of the tag data is valid and can be added. */
-	bool IsTagDataValid();
-	
 	/** Creates a new gameplay tag and adds it to the INI files based on the widget's stored parameters */
 	void CreateNewGameplayTag();
 
+	/** Handle validation and creation of restricted tags */
+	void ValidateNewRestrictedTag();
+	void CreateNewRestrictedGameplayTag();
+	void CancelNewTag();
+
 	/** Populates the widget's combo box with all potential places where a gameplay tag can be stored */
-	TSharedRef<SWidget> OnGenerateTagSourcesComboBox(TSharedPtr<FName> InItem);
+	TSharedRef<SWidget> OnGenerateTagSourcesComboBox(TSharedPtr<FString> InItem);
 
 	/** Creates the text displayed by the combo box when an option is selected */
 	FText CreateTagSourcesComboBoxContent() const;
@@ -96,7 +103,7 @@ private:
 private:
 
 	/** All potential INI files where a gameplay tag can be stored */
-	TArray<TSharedPtr<FName> > TagSources;
+	TArray<TSharedPtr<FString> > TagSources;
 
 	/** The name of the next gameplay tag to create */
 	TSharedPtr<SEditableTextBox> TagNameTextBox;
@@ -105,7 +112,8 @@ private:
 	TSharedPtr<SEditableTextBox> TagCommentTextBox;
 
 	/** The INI file where the next gameplay tag will be created */
-	TSharedPtr<SComboBox<TSharedPtr<FName>>> TagSourcesComboBox;
+	TSharedPtr<SSearchableComboBox> TagSourcesComboBox;
+
 
 	/** Callback for when a new gameplay tag has been added to the INI files */
 	FOnGameplayTagAdded OnGameplayTagAdded;
@@ -119,6 +127,9 @@ private:
 	/** Tracks if this widget should get keyboard focus */
 	bool bShouldGetKeyboardFocus = false;
 
+	/** If true, this widget is displaying restricted tags; if false this widget displays regular gameplay tags. */
+	bool bRestrictedTags = false;
+
 	/** Cached from NewTagName argument. */
 	FString DefaultNewName;
 
diff --git a/Engine/Plugins/Editor/GameplayTagsEditor/Source/GameplayTagsEditor/Private/SAddNewGameplayTagSourceWidget.cpp b/Engine/Plugins/Editor/GameplayTagsEditor/Source/GameplayTagsEditor/Private/SAddNewGameplayTagSourceWidget.cpp
--- a/Engine/Plugins/Editor/GameplayTagsEditor/Source/GameplayTagsEditor/Private/SAddNewGameplayTagSourceWidget.cpp
+++ b/Engine/Plugins/Editor/GameplayTagsEditor/Source/GameplayTagsEditor/Private/SAddNewGameplayTagSourceWidget.cpp
@@ -7,7 +7,7 @@
 #include "GameplayTagsManager.h"
 #include "Widgets/Input/SButton.h"
 #include "Misc/Paths.h"
-#include "Widgets/Input/SComboBox.h"
+#include "SSearchableComboBox.h"
 #include "Widgets/Input/SEditableTextBox.h"
 #include "Widgets/Layout/SGridPanel.h"
 
@@ -78,7 +78,7 @@ void SAddNewGameplayTagSourceWidget::Construct(const FArguments& InArgs)
 			.VAlign(VAlign_Center)
 			.HAlign(HAlign_Fill)
 			[
-				SAssignNew(TagRootsComboBox, SComboBox<TSharedPtr<FString> >)
+				SAssignNew(TagRootsComboBox, SSearchableComboBox)
 				.OptionsSource(&TagRoots)
 				.OnGenerateWidget(this, &SAddNewGameplayTagSourceWidget::OnGenerateTagRootsComboBox)
 				.Content()
diff --git a/Engine/Plugins/Editor/GameplayTagsEditor/Source/GameplayTagsEditor/Private/SAddNewGameplayTagSourceWidget.h b/Engine/Plugins/Editor/GameplayTagsEditor/Source/GameplayTagsEditor/Private/SAddNewGameplayTagSourceWidget.h
--- a/Engine/Plugins/Editor/GameplayTagsEditor/Source/GameplayTagsEditor/Private/SAddNewGameplayTagSourceWidget.h
+++ b/Engine/Plugins/Editor/GameplayTagsEditor/Source/GameplayTagsEditor/Private/SAddNewGameplayTagSourceWidget.h
@@ -5,7 +5,7 @@
 #include "Widgets/SCompoundWidget.h"
 
 class SEditableTextBox;
-template <typename OptionType> class SComboBox;
+class SSearchableComboBox;
 
 /** Widget allowing the user to create new gameplay tags */
 class SAddNewGameplayTagSourceWidget : public SCompoundWidget
@@ -60,7 +60,7 @@ private:
 	TArray<TSharedPtr<FString> > TagRoots;
 
 	/** The directory where the next tag source will be created */
-	TSharedPtr<SComboBox<TSharedPtr<FString> > > TagRootsComboBox;
+	TSharedPtr<SSearchableComboBox> TagRootsComboBox;
 
 	/** Tracks if this widget should get keyboard focus */
 	bool bShouldGetKeyboardFocus;
diff --git a/Engine/Plugins/Editor/GameplayTagsEditor/Source/GameplayTagsEditor/Private/SAddNewGameplayTagWidget.cpp b/Engine/Plugins/Editor/GameplayTagsEditor/Source/GameplayTagsEditor/Private/SAddNewGameplayTagWidget.cpp
--- a/Engine/Plugins/Editor/GameplayTagsEditor/Source/GameplayTagsEditor/Private/SAddNewGameplayTagWidget.cpp
+++ b/Engine/Plugins/Editor/GameplayTagsEditor/Source/GameplayTagsEditor/Private/SAddNewGameplayTagWidget.cpp
@@ -8,7 +8,7 @@
 #include "Misc/Paths.h"
 #include "Widgets/Input/SButton.h"
 #include "Misc/MessageDialog.h"
-#include "Widgets/Input/SComboBox.h"
+#include "SSearchableComboBox.h"
 #include "Widgets/Input/SEditableTextBox.h"
 #include "Widgets/Layout/SGridPanel.h"
 #include "Widgets/Notifications/SNotificationList.h"
@@ -35,6 +35,7 @@ void SAddNewGameplayTagWidget::Construct(const FArguments& InArgs)
 
 	bAddingNewTag = false;
 	bShouldGetKeyboardFocus = false;
+	bRestrictedTags = InArgs._RestrictedTags;
 
 	OnGameplayTagAdded = InArgs._OnGameplayTagAdded;
 	IsValidTag = InArgs._IsValidTag;
@@ -113,7 +114,7 @@ void SAddNewGameplayTagWidget::Construct(const FArguments& InArgs)
 				.FillWidth(1.0f)
 				.VAlign(VAlign_Center)
 				[
-					SAssignNew(TagSourcesComboBox, SComboBox<TSharedPtr<FName> >)
+					SAssignNew(TagSourcesComboBox, SSearchableComboBox)
 					.OptionsSource(&TagSources)
 					.OnGenerateWidget(this, &SAddNewGameplayTagWidget::OnGenerateTagSourcesComboBox)
 					.ToolTipText(this, &SAddNewGameplayTagWidget::CreateTagSourcesComboBoxToolTip)
@@ -166,8 +167,7 @@ EVisibility SAddNewGameplayTagWidget::OnGetTagSourceFavoritesVisibility() const
 
 FReply SAddNewGameplayTagWidget::OnToggleTagSourceFavoriteClicked()
 {
-	const bool bHasSelectedItem = TagSourcesComboBox.IsValid() && TagSourcesComboBox->GetSelectedItem().IsValid();
-	const FName ActiveTagSource = bHasSelectedItem ? *TagSourcesComboBox->GetSelectedItem().Get() : FName();
+	const FName ActiveTagSource = GetSelectedTagSource();
 	const bool bWasFavorite = FGameplayTagSource::GetFavoriteName() == ActiveTagSource;
 
 	FGameplayTagSource::SetFavoriteName(bWasFavorite ? NAME_None : ActiveTagSource);
@@ -177,8 +177,7 @@ FReply SAddNewGameplayTagWidget::OnToggleTagSourceFavoriteClicked()
 
 const FSlateBrush* SAddNewGameplayTagWidget::OnGetTagSourceFavoriteImage() const
 {
-	const bool bHasSelectedItem = TagSourcesComboBox.IsValid() && TagSourcesComboBox->GetSelectedItem().IsValid();
-	const FName ActiveTagSource = bHasSelectedItem ? *TagSourcesComboBox->GetSelectedItem().Get() : FName();
+	const FName ActiveTagSource = GetSelectedTagSource();
 	const bool bIsFavoriteTagSource = FGameplayTagSource::GetFavoriteName() == ActiveTagSource;
 
 	return FAppStyle::GetBrush(bIsFavoriteTagSource ? TEXT("Icons.Star") : TEXT("PropertyWindow.Favorites_Disabled"));
@@ -199,29 +198,51 @@ void SAddNewGameplayTagWidget::PopulateTagSources()
 	const UGameplayTagsManager& Manager = UGameplayTagsManager::Get();
 	TagSources.Empty();
 
-	const FName DefaultSource = FGameplayTagSource::GetDefaultName();
-
-	// Always ensure that the default source is first
-	TagSources.Add( MakeShareable( new FName( DefaultSource ) ) );
-
 	TArray<const FGameplayTagSource*> Sources;
-	Manager.FindTagSourcesWithType(EGameplayTagSourceType::TagList, Sources);
 
-	Algo::SortBy(Sources, &FGameplayTagSource::SourceName, FNameLexicalLess());
-
-	for (const FGameplayTagSource* Source : Sources)
+	if (bRestrictedTags)
 	{
-		if (Source != nullptr && Source->SourceName != DefaultSource)
+		Manager.GetRestrictedTagSources(Sources);
+
+		// Add the placeholder source if no other sources exist
+		if (Sources.Num() == 0)
 		{
-			TagSources.Add(MakeShareable(new FName(Source->SourceName)));
+			TagSources.Add(MakeShareable(new FString()));
 		}
-	}
 
-	//Set selection to the latest added source
-	if (TagSourcesComboBox != nullptr)
+		for (const FGameplayTagSource* Source : Sources)
+		{
+			if (Source != nullptr && !Source->SourceName.IsNone())
+			{
+				TagSources.Add(MakeShareable(new FString(Source->SourceName.ToString())));
+			}
+		}
+	}
+	else
 	{
-		TagSourcesComboBox->SetSelectedItem(TagSources.Last());
-	}	
+		const FName DefaultSource = FGameplayTagSource::GetDefaultName();
+
+		// Always ensure that the default source is first
+		TagSources.Add( MakeShareable( new FString( DefaultSource.ToString() ) ) );
+
+		Manager.FindTagSourcesWithType(EGameplayTagSourceType::TagList, Sources);
+
+		Algo::SortBy(Sources, &FGameplayTagSource::SourceName, FNameLexicalLess());
+
+		for (const FGameplayTagSource* Source : Sources)
+		{
+			if (Source != nullptr && Source->SourceName != DefaultSource)
+			{
+				TagSources.Add(MakeShareable(new FString(Source->SourceName.ToString())));
+			}
+		}
+
+		//Set selection to the latest added source
+		if (TagSourcesComboBox != nullptr)
+		{
+			TagSourcesComboBox->SetSelectedItem(TagSources.Last());
+		}
+	}
 }
 
 void SAddNewGameplayTagWidget::Reset(EResetType ResetType)
@@ -248,9 +269,9 @@ void SAddNewGameplayTagWidget::SelectTagSource(const FName& InSource)
 	{
 		for (int32 Index = 0; Index < TagSources.Num(); ++Index)
 		{
-			TSharedPtr<FName> Source = TagSources[Index];
+			TSharedPtr<FString> Source = TagSources[Index];
 
-			if (Source.IsValid() && *Source.Get() == InSource)
+			if (Source.IsValid() && *Source.Get() == InSource.ToString())
 			{
 				SourceIndex = Index;
 				break;
@@ -264,6 +285,19 @@ void SAddNewGameplayTagWidget::SelectTagSource(const FName& InSource)
 	}
 }
 
+FName SAddNewGameplayTagWidget::GetSelectedTagSource() const
+{
+	const bool bHasSelectedItem = TagSourcesComboBox.IsValid() && TagSourcesComboBox->GetSelectedItem().IsValid();
+
+	if (bHasSelectedItem)
+	{
+		// Convert to FName which the rest of the API expects
+		return FName(**TagSourcesComboBox->GetSelectedItem().Get());
+	}
+
+	return NAME_None;
+}
+
 void SAddNewGameplayTagWidget::OnCommitNewTagName(const FText& InText, ETextCommit::Type InCommitType)
 {
 	if (InCommitType == ETextCommit::OnEnter)
@@ -298,6 +332,12 @@ void SAddNewGameplayTagWidget::AddDuplicate(const FString& InParentTagName, cons
 
 void SAddNewGameplayTagWidget::CreateNewGameplayTag()
 {
+	if (bRestrictedTags)
+	{
+		ValidateNewRestrictedTag();
+		return;
+	}
+
 	if (NotificationItem.IsValid())
 	{
 		NotificationItem->SetVisibility(EVisibility::Collapsed);
@@ -325,7 +365,7 @@ void SAddNewGameplayTagWidget::CreateNewGameplayTag()
 	const FText TagNameAsText = TagNameTextBox->GetText();
 	FString TagName = TagNameAsText.ToString();
 	const FString TagComment = TagCommentTextBox->GetText().ToString();
-	const FName TagSource = *TagSourcesComboBox->GetSelectedItem().Get();
+	const FName TagSource = GetSelectedTagSource();
 
 	if (TagName.IsEmpty())
 	{
@@ -364,27 +404,144 @@ void SAddNewGameplayTagWidget::CreateNewGameplayTag()
 	Reset(EResetType::DoNotResetSource);
 }
 
-TSharedRef<SWidget> SAddNewGameplayTagWidget::OnGenerateTagSourcesComboBox(TSharedPtr<FName> InItem)
+
+void SAddNewGameplayTagWidget::ValidateNewRestrictedTag()
+{
+	UGameplayTagsManager& Manager = UGameplayTagsManager::Get();
+
+	FString TagName = TagNameTextBox->GetText().ToString();
+	FString TagComment = TagCommentTextBox->GetText().ToString();
+	const FName TagSource = GetSelectedTagSource();
+
+	if (TagSource == NAME_None)
+	{
+		FNotificationInfo Info(LOCTEXT("NoRestrictedSource", "You must specify a source file for restricted gameplay tags."));
+		Info.ExpireDuration = 10.f;
+		Info.bUseSuccessFailIcons = true;
+		Info.Image = FAppStyle::GetBrush(TEXT("MessageLog.Error"));
+
+		NotificationItem = FSlateNotificationManager::Get().AddNotification(Info);
+
+		return;
+	}
+
+	TArray<FString> TagSourceOwners;
+	Manager.GetOwnersForTagSource(TagSource.ToString(), TagSourceOwners);
+
+	bool bHasOwner = false;
+	for (const FString& Owner : TagSourceOwners)
+	{
+		if (!Owner.IsEmpty())
+		{
+			bHasOwner = true;
+			break;
+		}
+	}
+
+	if (bHasOwner)
+	{
+		// check if we're one of the owners; if we are then we don't need to pop up the permission dialog
+		bool bRequiresPermission = true;
+		const FString& UserName = FPlatformProcess::UserName();
+		for (const FString& Owner : TagSourceOwners)
+		{
+			if (Owner.Equals(UserName))
+			{
+				CreateNewRestrictedGameplayTag();
+				bRequiresPermission = false;
+			}
+		}
+
+		if (bRequiresPermission)
+		{
+			FString StringToDisplay = TEXT("Do you have permission from ");
+			StringToDisplay.Append(TagSourceOwners[0]);
+			for (int Idx = 1; Idx < TagSourceOwners.Num(); ++Idx)
+			{
+				StringToDisplay.Append(TEXT(" or "));
+				StringToDisplay.Append(TagSourceOwners[Idx]);
+			}
+			StringToDisplay.Append(TEXT(" to modify "));
+			StringToDisplay.Append(TagSource.ToString());
+			StringToDisplay.Append(TEXT("?"));
+
+			FNotificationInfo Info(FText::FromString(StringToDisplay));
+			Info.ExpireDuration = 10.f;
+			Info.ButtonDetails.Add(FNotificationButtonInfo(LOCTEXT("RestrictedTagPopupButtonAccept", "Yes"), FText(), FSimpleDelegate::CreateSP(this, &SAddNewGameplayTagWidget::CreateNewRestrictedGameplayTag), SNotificationItem::CS_None));
+			Info.ButtonDetails.Add(FNotificationButtonInfo(LOCTEXT("RestrictedTagPopupButtonReject", "No"), FText(), FSimpleDelegate::CreateSP(this, &SAddNewGameplayTagWidget::CancelNewTag), SNotificationItem::CS_None));
+
+			NotificationItem = FSlateNotificationManager::Get().AddNotification(Info);
+		}
+	}
+	else
+	{
+		CreateNewRestrictedGameplayTag();
+	}
+}
+
+void SAddNewGameplayTagWidget::CreateNewRestrictedGameplayTag()
+{
+	if (NotificationItem.IsValid())
+	{
+		NotificationItem->SetVisibility(EVisibility::Collapsed);
+	}
+
+	const UGameplayTagsManager& Manager = UGameplayTagsManager::Get();
+
+	// Only support adding tags via ini file
+	if (Manager.ShouldImportTagsFromINI() == false)
+	{
+		return;
+	}
+
+	const FString TagName = TagNameTextBox->GetText().ToString();
+	const FString TagComment = TagCommentTextBox->GetText().ToString();
+	const bool bAllowNonRestrictedChildren = true; // can be changed later
+	const FName TagSource = GetSelectedTagSource();
+
+	if (TagName.IsEmpty())
+	{
+		return;
+	}
+
+	// set bIsAddingNewTag, this guards against the window closing when it loses focus due to source control checking out a file
+	TGuardValue<bool>	Guard(bAddingNewTag, true);
+
+	IGameplayTagsEditorModule::Get().AddNewGameplayTagToINI(TagName, TagComment, TagSource, true, bAllowNonRestrictedChildren);
+
+	OnGameplayTagAdded.ExecuteIfBound(TagName, TagComment, TagSource);
+
+	Reset(EResetType::DoNotResetSource);
+}
+
+void SAddNewGameplayTagWidget::CancelNewTag()
+{
+	if (NotificationItem.IsValid())
+	{
+		NotificationItem->SetVisibility(EVisibility::Collapsed);
+	}
+}
+
+TSharedRef<SWidget> SAddNewGameplayTagWidget::OnGenerateTagSourcesComboBox(TSharedPtr<FString> InItem)
 {
 	return SNew(STextBlock)
-		.Text(FText::FromName(*InItem.Get()));
+		.Text(FText::AsCultureInvariant(*InItem.Get()));
 }
 
 FText SAddNewGameplayTagWidget::CreateTagSourcesComboBoxContent() const
 {
 	const bool bHasSelectedItem = TagSourcesComboBox.IsValid() && TagSourcesComboBox->GetSelectedItem().IsValid();
 
-	return bHasSelectedItem ? FText::FromName(*TagSourcesComboBox->GetSelectedItem().Get()) : LOCTEXT("NewTagLocationNotSelected", "Not selected");
+	return bHasSelectedItem ? FText::AsCultureInvariant(*TagSourcesComboBox->GetSelectedItem().Get()) : LOCTEXT("NewTagLocationNotSelected", "Not selected");
 }
 
 FText SAddNewGameplayTagWidget::CreateTagSourcesComboBoxToolTip() const
 {
-	const bool bHasSelectedItem = TagSourcesComboBox.IsValid() && TagSourcesComboBox->GetSelectedItem().IsValid();
-
-	if (bHasSelectedItem)
+	const FName ActiveTagSource = GetSelectedTagSource();
+	if (!ActiveTagSource.IsNone())
 	{
 		UGameplayTagsManager& Manager = UGameplayTagsManager::Get();
-		const FGameplayTagSource* Source = Manager.FindTagSource(*TagSourcesComboBox->GetSelectedItem().Get());
+		const FGameplayTagSource* Source = Manager.FindTagSource(ActiveTagSource);
 		if (Source)
 		{
 			FString FilePath = Source->GetConfigFileName();
diff --git a/Engine/Plugins/Editor/GameplayTagsEditor/Source/GameplayTagsEditor/Private/SAddNewRestrictedGameplayTagWidget.cpp b/Engine/Plugins/Editor/GameplayTagsEditor/Source/GameplayTagsEditor/Private/SAddNewRestrictedGameplayTagWidget.cpp
deleted file mode 100644
diff --git a/Engine/Plugins/Editor/GameplayTagsEditor/Source/GameplayTagsEditor/Private/SAddNewRestrictedGameplayTagWidget.h b/Engine/Plugins/Editor/GameplayTagsEditor/Source/GameplayTagsEditor/Private/SAddNewRestrictedGameplayTagWidget.h
deleted file mode 100644
diff --git a/Engine/Plugins/Editor/GameplayTagsEditor/Source/GameplayTagsEditor/Private/SGameplayTagPicker.cpp b/Engine/Plugins/Editor/GameplayTagsEditor/Source/GameplayTagsEditor/Private/SGameplayTagPicker.cpp
--- a/Engine/Plugins/Editor/GameplayTagsEditor/Source/GameplayTagsEditor/Private/SGameplayTagPicker.cpp
+++ b/Engine/Plugins/Editor/GameplayTagsEditor/Source/GameplayTagsEditor/Private/SGameplayTagPicker.cpp
@@ -17,7 +17,6 @@
 #include "GameplayTagsEditorModule.h"
 #include "Framework/Application/SlateApplication.h"
 #include "SAddNewGameplayTagWidget.h"
-#include "SAddNewRestrictedGameplayTagWidget.h"
 #include "SRenameGameplayTagDialog.h"
 #include "AssetRegistry/AssetData.h"
 #include "Editor.h"
@@ -142,8 +141,6 @@ void SGameplayTagPicker::Construct(const FArguments& InArgs)
 			}
 		}
 	}
-
-	const FText NewTagText = bRestrictedTags ? LOCTEXT("AddNewRestrictedTag", "Add New Restricted Gameplay Tag") : LOCTEXT("AddNewTag", "Add New Gameplay Tag");
 	
 	TSharedPtr<SComboButton> SettingsCombo = SNew(SComboButton)
 		.VAlign(VAlign_Center)
@@ -256,6 +253,7 @@ void SGameplayTagPicker::Construct(const FArguments& InArgs)
 				SAssignNew(AddNewTagWidget, SAddNewGameplayTagWidget )
 				.Padding(FMargin(0, 2))
 				.AddButtonPadding(FMargin(0, 4, 0, 0))
+				.RestrictedTags(bRestrictedTags)
 				.Visibility_Lambda([WeakSelf]()
 				{
 					const TSharedPtr<SGameplayTagPicker> Self = WeakSelf.Pin();
@@ -1058,99 +1056,6 @@ void SGameplayTagPicker::OnCollapseAllClicked(TSharedPtr<SComboButton> OwnerComb
 	}
 }
 
-void SGameplayTagPicker::OpenAddTagDialog(const EGameplayTagAdd Mode, TSharedPtr<FGameplayTagNode> InTagNode)
-{
-	TSharedPtr<SWindow> NewTagWindow = SNew(SWindow)
-		.Title(LOCTEXT("EditTagWindowTitle", "Edit Gameplay Tag"))
-		.SizingRule(ESizingRule::Autosized)
-		.SupportsMaximize(false)
-		.SupportsMinimize(false);
-
-	UGameplayTagsManager& Manager = UGameplayTagsManager::Get();
-	
-	FString TagName;
-	FName TagFName;
-	FString TagComment;
-	FName TagSource;
-	bool bTagIsExplicit;
-	bool bTagIsRestricted;
-	bool bTagAllowsNonRestrictedChildren;
-
-	if (InTagNode.IsValid())
-	{
-		TagName = InTagNode->GetCompleteTagString();
-		TagFName = InTagNode->GetCompleteTagName();
-	}
-
-	Manager.GetTagEditorData(TagFName, TagComment, TagSource, bTagIsExplicit, bTagIsRestricted, bTagAllowsNonRestrictedChildren);
-
-	TWeakPtr<SGameplayTagPicker> WeakSelf = StaticCastWeakPtr<SGameplayTagPicker>(AsWeak());
-	
-	if (!bRestrictedTags)
-    {
-		TSharedRef<SAddNewGameplayTagWidget> AddNewTagDialog = SNew(SAddNewGameplayTagWidget)
-			.OnGameplayTagAdded_Lambda([WeakSelf, NewTagWindow](const FString& TagName, const FString& TagComment, const FName& TagSource)
-			{
-				if (const TSharedPtr<SGameplayTagPicker> Self = WeakSelf.Pin())
-				{
-					Self->OnGameplayTagAdded(TagName, TagComment, TagSource);
-				}
-				if (NewTagWindow.IsValid())
-				{
-					NewTagWindow->RequestDestroyWindow();
-				}
-			});
-
-		if (Mode == EGameplayTagAdd::Child || Mode == EGameplayTagAdd::Root)
-		{
-			AddNewTagDialog->AddSubtagFromParent(TagName, TagSource);
-		}
-		else if (Mode == EGameplayTagAdd::Duplicate)
-		{
-			AddNewTagDialog->AddDuplicate(TagName, TagSource);
-		}
-
-		NewTagWindow->SetContent(SNew(SBox)
-			.MinDesiredWidth(320.0f)
-			[
-				AddNewTagDialog
-			]);
-    }
-    else if (bRestrictedTags)
-    {
-        TSharedRef<SAddNewRestrictedGameplayTagWidget> AddNewTagDialog = SNew(SAddNewRestrictedGameplayTagWidget)
-			.OnRestrictedGameplayTagAdded_Lambda([WeakSelf, NewTagWindow](const FString& TagName, const FString& TagComment, const FName& TagSource)
-			{
-				if (const TSharedPtr<SGameplayTagPicker> Self = WeakSelf.Pin())
-				{
-					Self->OnGameplayTagAdded(TagName, TagComment, TagSource);
-				}
-				if (NewTagWindow.IsValid())
-				{
-					NewTagWindow->RequestDestroyWindow();
-				}
-			});
-
-    	if (Mode == EGameplayTagAdd::Child || Mode == EGameplayTagAdd::Root)
-    	{
-	        AddNewTagDialog->AddSubtagFromParent(TagName, TagSource, bTagAllowsNonRestrictedChildren);
-    	}
-    	else if (Mode == EGameplayTagAdd::Duplicate)
-    	{
-    		AddNewTagDialog->AddDuplicate(TagName, TagSource, bTagAllowsNonRestrictedChildren);
-    	}
-        
-    	NewTagWindow->SetContent(SNew(SBox)
-			.MinDesiredWidth(320.0f)
-			[
-				AddNewTagDialog
-			]);
-    }
-	
-	TSharedPtr<SWindow> CurrentWindow = FSlateApplication::Get().FindWidgetWindow(AsShared());
-	FSlateApplication::Get().AddModalWindow(NewTagWindow.ToSharedRef(), CurrentWindow);
-}
-
 void SGameplayTagPicker::ShowInlineAddTagWidget(const EGameplayTagAdd Mode, TSharedPtr<FGameplayTagNode> InTagNode)
 {
 	UGameplayTagsManager& Manager = UGameplayTagsManager::Get();
@@ -1185,13 +1090,6 @@ void SGameplayTagPicker::ShowInlineAddTagWidget(const EGameplayTagAdd Mode, TSha
 	}
 }
 
-FReply SGameplayTagPicker::OnAddRootTagClicked()
-{
-	OpenAddTagDialog(EGameplayTagAdd::Root);
-	
-	return FReply::Handled();
-}
-
 TSharedRef<SWidget> SGameplayTagPicker::MakeTagActionsMenu(TSharedPtr<FGameplayTagNode> InTagNode, TSharedPtr<SComboButton> ActionsCombo, bool bInShouldCloseWindowAfterMenuSelection)
 {
 	if (!InTagNode.IsValid())
@@ -1904,7 +1802,14 @@ TWeakPtr<SGameplayTagPicker> OpenGameplayTagManager(const FGameplayTagManagerWin
 		FText Title = Args.Title;
 		if (Title.IsEmpty())
 		{
-			Title = LOCTEXT("GameplayTagPicker_ManagerTitle", "Gameplay Tag Manager");
+			if (Args.bRestrictedTags)
+			{
+				Title = LOCTEXT("GameplayTagPicker_RestrictedTitle", "Restricted Gameplay Tags");
+			}
+			else
+			{
+				Title = LOCTEXT("GameplayTagPicker_ManagerTitle", "Gameplay Tag Manager");
+			}
 		}
 		
 		GameplayTagPickerWindow = SNew(SWindow)
diff --git a/Engine/Plugins/Editor/GameplayTagsEditor/Source/GameplayTagsEditor/Private/SGameplayTagWidget.cpp b/Engine/Plugins/Editor/GameplayTagsEditor/Source/GameplayTagsEditor/Private/SGameplayTagWidget.cpp
--- a/Engine/Plugins/Editor/GameplayTagsEditor/Source/GameplayTagsEditor/Private/SGameplayTagWidget.cpp
+++ b/Engine/Plugins/Editor/GameplayTagsEditor/Source/GameplayTagsEditor/Private/SGameplayTagWidget.cpp
@@ -19,7 +19,6 @@
 #include "Framework/Application/SlateApplication.h"
 #include "SAddNewGameplayTagWidget.h"
 #include "SAddNewGameplayTagSourceWidget.h"
-#include "SAddNewRestrictedGameplayTagWidget.h"
 #include "SRenameGameplayTagDialog.h"
 #include "Editor.h"
 #include "Framework/MultiBox/MultiBoxBuilder.h"
@@ -182,20 +181,7 @@ void SGameplayTagWidget::Construct(const FArguments& InArgs, const TArray<FEdita
 				]
 			]
 
-
-			// Expandable UI for adding restricted tags
-			+SVerticalBox::Slot()
-			.AutoHeight()
-			.VAlign(VAlign_Top)
-			.Padding(16.0f, 0.0f)
-			[
-				SAssignNew(AddNewRestrictedTagWidget, SAddNewRestrictedGameplayTagWidget)
-				.Visibility(this, &SGameplayTagWidget::DetermineAddNewRestrictedTagWidgetVisibility)
-				.OnRestrictedGameplayTagAdded(this, &SGameplayTagWidget::OnGameplayTagAdded)
-				.NewRestrictedTagName(InArgs._NewTagName)
-			]
-
-			// Expandable UI for adding non-restricted tags
+			// Expandable UI for adding tags
 			+SVerticalBox::Slot()
 			.AutoHeight()
 			.VAlign(VAlign_Top)
@@ -205,6 +191,7 @@ void SGameplayTagWidget::Construct(const FArguments& InArgs, const TArray<FEdita
 				.Visibility(this, &SGameplayTagWidget::DetermineAddNewTagWidgetVisibility)
 				.OnGameplayTagAdded(this, &SGameplayTagWidget::OnGameplayTagAdded)
 				.NewTagName(InArgs._NewTagName)
+				.RestrictedTags(bRestrictedTags)
 			]
 
 			// Expandable UI controls
@@ -893,16 +880,11 @@ FReply SGameplayTagWidget::OnAddSubtagClicked(TSharedPtr<FGameplayTagNode> InTag
 
 		Manager.GetTagEditorData(InTagNode->GetCompleteTagName(), TagComment, TagSource, bTagIsExplicit, bTagIsRestricted, bTagAllowsNonRestrictedChildren);
 
-		if (!bRestrictedTags && AddNewTagWidget.IsValid())
+		if (AddNewTagWidget.IsValid())
 		{
 			bAddTagSectionExpanded = true; 
 			AddNewTagWidget->AddSubtagFromParent(TagName, TagSource);
 		}
-		else if (bRestrictedTags && AddNewRestrictedTagWidget.IsValid())
-		{
-			bAddTagSectionExpanded = true;
-			AddNewRestrictedTagWidget->AddSubtagFromParent(TagName, TagSource, bTagAllowsNonRestrictedChildren);
-		}
 	}
 	return FReply::Handled();
 }
@@ -1295,19 +1277,6 @@ EVisibility SGameplayTagWidget::DetermineAddNewTagWidgetVisibility() const
 	return EVisibility::Visible;
 }
 
-EVisibility SGameplayTagWidget::DetermineAddNewRestrictedTagWidgetVisibility() const
-{
-	UGameplayTagsManager& Manager = UGameplayTagsManager::Get();
-
-	if (!Manager.ShouldImportTagsFromINI() || !bAddTagSectionExpanded || !bRestrictedTags || bForceHideAddNewTag )
-	{
-		// If we can't support adding tags from INI files, we should never see this widget
-		return EVisibility::Collapsed;
-	}
-
-	return EVisibility::Visible;
-}
-
 EVisibility SGameplayTagWidget::DetermineAddNewSourceWidgetVisibility() const
 {
 	UGameplayTagsManager& Manager = UGameplayTagsManager::Get();
diff --git a/Engine/Plugins/Editor/GameplayTagsEditor/Source/GameplayTagsEditor/Public/SGameplayTagPicker.h b/Engine/Plugins/Editor/GameplayTagsEditor/Source/GameplayTagsEditor/Public/SGameplayTagPicker.h
--- a/Engine/Plugins/Editor/GameplayTagsEditor/Source/GameplayTagsEditor/Public/SGameplayTagPicker.h
+++ b/Engine/Plugins/Editor/GameplayTagsEditor/Source/GameplayTagsEditor/Public/SGameplayTagPicker.h
@@ -1,4 +1,4 @@
-﻿// Copyright Epic Games, Inc. All Rights Reserved.
+// Copyright Epic Games, Inc. All Rights Reserved.
 
 #pragma once
 
@@ -317,14 +317,9 @@ private:
 	/** Callback for when a new tag is added */
 	void OnGameplayTagAdded(const FString& TagName, const FString& TagComment, const FName& TagSource);
 
-	/** Opens add tag modal dialog. */
-	void OpenAddTagDialog(const EGameplayTagAdd Mode, TSharedPtr<FGameplayTagNode> InTagNode = TSharedPtr<FGameplayTagNode>());
-
 	/** Initializes the inline add tag widget and makes in visible. */
 	void ShowInlineAddTagWidget(const EGameplayTagAdd Mode, TSharedPtr<FGameplayTagNode> InTagNode = TSharedPtr<FGameplayTagNode>());
 
-	FReply OnAddRootTagClicked();
-
 	FText GetHighlightText() const;
 	
 	/** Creates a dropdown menu to provide additional functionality for tags (renaming, deletion, search for references, etc.) */
diff --git a/Engine/Plugins/Editor/GameplayTagsEditor/Source/GameplayTagsEditor/Public/SGameplayTagWidget.h b/Engine/Plugins/Editor/GameplayTagsEditor/Source/GameplayTagsEditor/Public/SGameplayTagWidget.h
--- a/Engine/Plugins/Editor/GameplayTagsEditor/Source/GameplayTagsEditor/Public/SGameplayTagWidget.h
+++ b/Engine/Plugins/Editor/GameplayTagsEditor/Source/GameplayTagsEditor/Public/SGameplayTagWidget.h
@@ -204,9 +204,6 @@ private:
 	/** The widget that controls how new gameplay tags are added to the config files */
 	TSharedPtr<class SAddNewGameplayTagWidget> AddNewTagWidget;
 
-	/** The widget that controls how new restricted gameplay tags are added to the config files */
-	TSharedPtr<class SAddNewRestrictedGameplayTagWidget> AddNewRestrictedTagWidget;
-
 	/** The widget that controls how new gameplay tag source files are added */
 	TSharedPtr<class SAddNewGameplayTagSourceWidget> AddNewTagSourceWidget;
 
@@ -354,9 +351,6 @@ private:
 	/** Helper function to determine the visibility of the Add New Tag widget */
 	EVisibility DetermineAddNewTagWidgetVisibility() const;
 
-	/** Helper function to determine the visibility of the Add New Tag widget */
-	EVisibility DetermineAddNewRestrictedTagWidgetVisibility() const;
-
 	/** Helper function to determine the visibility of the Add New Tag Source widget */
 	EVisibility DetermineAddNewSourceWidgetVisibility() const;
 
