Fix add tag window to correctly show sources, there was an uninitialized bool which randomly hid them

diff --git a/Engine/Plugins/Editor/GameplayTagsEditor/Source/GameplayTagsEditor/Internal/SAddNewGameplayTagWidget.h b/Engine/Plugins/Editor/GameplayTagsEditor/Source/GameplayTagsEditor/Internal/SAddNewGameplayTagWidget.h
--- a/Engine/Plugins/Editor/GameplayTagsEditor/Source/GameplayTagsEditor/Internal/SAddNewGameplayTagWidget.h
+++ b/Engine/Plugins/Editor/GameplayTagsEditor/Source/GameplayTagsEditor/Internal/SAddNewGameplayTagWidget.h
@@ -29,6 +29,7 @@ public:
 		: _NewTagName(TEXT(""))
 		, _Padding(FMargin(15))
 		, _AddButtonPadding(FMargin(0, 16, 0, 0))
+		, _RestrictedTags(false)
 	{}
 		SLATE_EVENT(FOnGameplayTagAdded, OnGameplayTagAdded)	// Callback for when a new tag is added	
 		SLATE_EVENT(FIsValidTag, IsValidTag)
