PR #11519: Gameplay Tags: Added the DevComment to the tooltip for tag chips in both SGameplayTagCombo and SGameplayTagContainerCombo

diff --git a/Engine/Plugins/Editor/GameplayTagsEditor/Source/GameplayTagsEditor/Private/SGameplayTagCombo.cpp b/Engine/Plugins/Editor/GameplayTagsEditor/Source/GameplayTagsEditor/Private/SGameplayTagCombo.cpp
--- a/Engine/Plugins/Editor/GameplayTagsEditor/Source/GameplayTagsEditor/Private/SGameplayTagCombo.cpp
+++ b/Engine/Plugins/Editor/GameplayTagsEditor/Source/GameplayTagsEditor/Private/SGameplayTagCombo.cpp
@@ -1,4 +1,4 @@
-﻿// Copyright Epic Games, Inc. All Rights Reserved.
+// Copyright Epic Games, Inc. All Rights Reserved.
 
 #include "SGameplayTagCombo.h"
 #include "DetailLayoutBuilder.h"
@@ -145,11 +145,25 @@ FText SGameplayTagCombo::GetText() const
 
 FText SGameplayTagCombo::GetToolTipText() const
 {
+	FGameplayTag TagToDisplay;
 	if (PropertyHandle.IsValid())
 	{
-		return TagsFromProperty.IsEmpty() ? FText::GetEmpty() : FText::FromName(TagsFromProperty[0].GetTagName());
+		TagToDisplay = (TagsFromProperty.Num() > 0) ? TagsFromProperty[0] : FGameplayTag();
+	}
+	else
+	{
+		TagToDisplay = TagAttribute.Get();
+	}
+
+	TSharedPtr<FGameplayTagNode> TagNode = UGameplayTagsManager::Get().FindTagNode(TagToDisplay);
+	if (TagNode.IsValid() && !TagNode->GetDevComment().IsEmpty())
+	{
+		return FText::FromString(FString::Printf(TEXT("%s\n\n%s"), *TagToDisplay.ToString(), *TagNode->GetDevComment()));
+	}
+	else
+	{
+		return FText::FromName(TagToDisplay.GetTagName());
 	}
-	return FText::FromName(TagAttribute.Get().GetTagName());
 }
 
 bool SGameplayTagCombo::IsSelected() const
@@ -396,4 +410,4 @@ void SGameplayTagCombo::RefreshTagsFromProperty()
 	}
 }
 
-#undef LOCTEXT_NAMESPACE
+#undef LOCTEXT_NAMESPACE
\ No newline at end of file
diff --git a/Engine/Plugins/Editor/GameplayTagsEditor/Source/GameplayTagsEditor/Private/SGameplayTagContainerCombo.cpp b/Engine/Plugins/Editor/GameplayTagsEditor/Source/GameplayTagsEditor/Private/SGameplayTagContainerCombo.cpp
--- a/Engine/Plugins/Editor/GameplayTagsEditor/Source/GameplayTagsEditor/Private/SGameplayTagContainerCombo.cpp
+++ b/Engine/Plugins/Editor/GameplayTagsEditor/Source/GameplayTagsEditor/Private/SGameplayTagContainerCombo.cpp
@@ -176,7 +176,7 @@ TSharedRef<ITableRow> SGameplayTagContainerCombo::MakeTagListViewRow(TSharedPtr<
 			.ReadOnly(bIsReadOnly)
 			.ShowClearButton(true)
 			.Text(FText::FromName(Item->Tag.GetTagName()))
-			.ToolTipText(FText::FromName(Item->Tag.GetTagName()))
+			.ToolTipText(this, &SGameplayTagContainerCombo::GetTagToolTipText, Item->Tag)
 			.IsSelected(!Item->bMultipleValues)
 			.OnClearPressed(this, &SGameplayTagContainerCombo::OnClearTagClicked, Item->Tag)
 			.OnEditPressed(this, &SGameplayTagContainerCombo::OnEditClicked, Item->Tag)
@@ -230,6 +230,20 @@ TSharedRef<SWidget> SGameplayTagContainerCombo::OnGetMenuContent()
 	return TagPicker.ToSharedRef();
 }
 
+FText SGameplayTagContainerCombo::GetTagToolTipText(const FGameplayTag TagToDisplay) const
+{
+	TSharedPtr<FGameplayTagNode> TagNode = UGameplayTagsManager::Get().FindTagNode(TagToDisplay);
+
+	if (TagNode.IsValid() && !TagNode->GetDevComment().IsEmpty())
+	{
+		return FText::FromString(FString::Printf(TEXT("%s\n\n%s"), *TagToDisplay.ToString(), *TagNode->GetDevComment()));
+	}
+	else
+	{
+		return FText::FromName(TagToDisplay.GetTagName());
+	}
+}
+
 FReply SGameplayTagContainerCombo::OnTagMenu(const FPointerEvent& MouseEvent, const FGameplayTag GameplayTag)
 {
 	FMenuBuilder MenuBuilder(/*bShouldCloseWindowAfterMenuSelection=*/ true, /*CommandList=*/ nullptr);
@@ -634,4 +648,4 @@ void SGameplayTagContainerCombo::RefreshTagContainers()
 	}
 }
 
-#undef LOCTEXT_NAMESPACE
+#undef LOCTEXT_NAMESPACE
\ No newline at end of file
diff --git a/Engine/Plugins/Editor/GameplayTagsEditor/Source/GameplayTagsEditor/Public/SGameplayTagContainerCombo.h b/Engine/Plugins/Editor/GameplayTagsEditor/Source/GameplayTagsEditor/Public/SGameplayTagContainerCombo.h
--- a/Engine/Plugins/Editor/GameplayTagsEditor/Source/GameplayTagsEditor/Public/SGameplayTagContainerCombo.h
+++ b/Engine/Plugins/Editor/GameplayTagsEditor/Source/GameplayTagsEditor/Public/SGameplayTagContainerCombo.h
@@ -1,4 +1,4 @@
-﻿// Copyright Epic Games, Inc. All Rights Reserved.
+// Copyright Epic Games, Inc. All Rights Reserved.
 
 #pragma once
 
@@ -89,7 +89,8 @@ private:
 	TSharedRef<ITableRow> MakeTagListViewRow(TSharedPtr<FEditableItem> Item, const TSharedRef<STableViewBase>& OwnerTable);
 	void OnMenuOpenChanged(const bool bOpen);
 	TSharedRef<SWidget> OnGetMenuContent();
-	
+	FText GetTagToolTipText(const FGameplayTag TagToDisplay) const;
+
 	FReply OnTagMenu(const FPointerEvent& MouseEvent, const FGameplayTag GameplayTag);
 	FReply OnEmptyMenu(const FPointerEvent& MouseEvent);
 	FReply OnEditClicked(const FGameplayTag TagToHilight);
@@ -116,4 +117,4 @@ private:
 	TSharedPtr<SComboButton> ComboButton;
 	TSharedPtr<SGameplayTagPicker> TagPicker;
 	FGameplayTag TagToHilight;
-};
+};
\ No newline at end of file
diff --git a/Engine/Source/Runtime/GameplayTags/Classes/GameplayTagsManager.h b/Engine/Source/Runtime/GameplayTags/Classes/GameplayTagsManager.h
--- a/Engine/Source/Runtime/GameplayTags/Classes/GameplayTagsManager.h
+++ b/Engine/Source/Runtime/GameplayTags/Classes/GameplayTagsManager.h
@@ -256,6 +256,11 @@ struct FGameplayTagNode
 	const TArray<FName>& GetAllSourceNames() const { return SourceNames; }
 #endif
 
+#if WITH_EDITORONLY_DATA
+	/** Returns the Comment for this tag */
+	FString GetDevComment() const { return DevComment; }
+#endif
+
 private:
 	/** Raw name for this tag at current rank in the tree */
 	FName Tag;
@@ -961,4 +966,4 @@ private:
 
 	GAMEPLAYTAGS_API const static FName NAME_Categories;
 	GAMEPLAYTAGS_API const static FName NAME_GameplayTagFilter;
-};
+};
\ No newline at end of file
