GameplayTag Picker: If there is only one root node, open it by default (very useful for Categories meta tag)

#rnx

diff --git a/Engine/Plugins/Editor/GameplayTagsEditor/Source/GameplayTagsEditor/Private/SGameplayTagPicker.cpp b/Engine/Plugins/Editor/GameplayTagsEditor/Source/GameplayTagsEditor/Private/SGameplayTagPicker.cpp
--- a/Engine/Plugins/Editor/GameplayTagsEditor/Source/GameplayTagsEditor/Private/SGameplayTagPicker.cpp
+++ b/Engine/Plugins/Editor/GameplayTagsEditor/Source/GameplayTagsEditor/Private/SGameplayTagPicker.cpp
@@ -1090,6 +1090,11 @@ void SGameplayTagPicker::ShowInlineAddTagWidget(const EGameplayTagAdd Mode, TSha
 	}
 }
 
+bool SGameplayTagPicker::IsNodeSingleRoot(const TSharedPtr<FGameplayTagNode>& InNode) const
+{
+	return TagItems.Num() == 1 && TagItems[0] == InNode;
+}
+
 TSharedRef<SWidget> SGameplayTagPicker::MakeTagActionsMenu(TSharedPtr<FGameplayTagNode> InTagNode, TSharedPtr<SComboButton> ActionsCombo, bool bInShouldCloseWindowAfterMenuSelection)
 {
 	if (!InTagNode.IsValid())
@@ -1420,7 +1425,8 @@ void SGameplayTagPicker::SetDefaultTagNodeItemExpansion(TSharedPtr<FGameplayTagN
 	TGuardValue<bool> PersistExpansionChangeGuard(bPersistExpansionChange, false);
 	if (Node.IsValid() && TagTreeWidget.IsValid())
 	{
-		const bool bIsExpanded = CachedExpandedItems.Contains(Node) || IsTagChecked(Node) == ECheckBoxState::Checked;
+		const bool bIsExpanded = CachedExpandedItems.Contains(Node) || IsTagChecked(Node) == ECheckBoxState::Checked || IsNodeSingleRoot(Node);
+
 		TagTreeWidget->SetItemExpansion(Node, bIsExpanded);
 		
 		const TArray<TSharedPtr<FGameplayTagNode>>& ChildTags = Node->GetChildTagNodes();
@@ -1438,7 +1444,11 @@ void SGameplayTagPicker::LoadTagNodeItemExpansion(TSharedPtr<FGameplayTagNode> N
 	{
 		bool bIsExpanded = false;
 
-		if (GConfig->GetBool(*SettingsIniSection, *(SettingsName + Node->GetCompleteTagString() + TEXT(".Expanded")), bIsExpanded, GetGameplayTagsEditorStateIni()))
+		if (IsNodeSingleRoot(Node)) // Single roots are always expanded
+		{
+			TagTreeWidget->SetItemExpansion(Node, true);
+		}
+		else if (GConfig->GetBool(*SettingsIniSection, *(SettingsName + Node->GetCompleteTagString() + TEXT(".Expanded")), bIsExpanded, GetGameplayTagsEditorStateIni()))
 		{
 			TagTreeWidget->SetItemExpansion(Node, bIsExpanded);
 			if (bIsExpanded)
diff --git a/Engine/Plugins/Editor/GameplayTagsEditor/Source/GameplayTagsEditor/Public/SGameplayTagPicker.h b/Engine/Plugins/Editor/GameplayTagsEditor/Source/GameplayTagsEditor/Public/SGameplayTagPicker.h
--- a/Engine/Plugins/Editor/GameplayTagsEditor/Source/GameplayTagsEditor/Public/SGameplayTagPicker.h
+++ b/Engine/Plugins/Editor/GameplayTagsEditor/Source/GameplayTagsEditor/Public/SGameplayTagPicker.h
@@ -320,6 +320,9 @@ private:
 	/** Initializes the inline add tag widget and makes in visible. */
 	void ShowInlineAddTagWidget(const EGameplayTagAdd Mode, TSharedPtr<FGameplayTagNode> InTagNode = TSharedPtr<FGameplayTagNode>());
 
+	/** Returns true if the passed in node is the root node, and there is only one node on the root level */
+	bool IsNodeSingleRoot(const TSharedPtr<FGameplayTagNode>& InNode) const;
+
 	FText GetHighlightText() const;
 	
 	/** Creates a dropdown menu to provide additional functionality for tags (renaming, deletion, search for references, etc.) */
