Replace instances of String.Append(FString::Printf(...)) with String.Appendf(...) to avoid temporary string allocations.

diff --git a/Engine/Plugins/Editor/GameplayTagsEditor/Source/GameplayTagsEditor/Private/SGameplayTagPicker.cpp b/Engine/Plugins/Editor/GameplayTagsEditor/Source/GameplayTagsEditor/Private/SGameplayTagPicker.cpp
--- a/Engine/Plugins/Editor/GameplayTagsEditor/Source/GameplayTagsEditor/Private/SGameplayTagPicker.cpp
+++ b/Engine/Plugins/Editor/GameplayTagsEditor/Source/GameplayTagsEditor/Private/SGameplayTagPicker.cpp
@@ -606,12 +606,12 @@ TSharedRef<ITableRow> SGameplayTagPicker::OnGenerateRow(TSharedPtr<FGameplayTagN
 
 			bIsExplicitTag = Node->bIsExplicitTag;
 
-			TooltipString.Append(FString::Printf(TEXT("\n(%s%s)"), bIsExplicitTag ? TEXT("") : TEXT("Implicit "), *AllSources));
+			TooltipString.Appendf(TEXT("\n(%s%s)"), bIsExplicitTag ? TEXT("") : TEXT("Implicit "), *AllSources);
 
 			// tag comments
 			if (!Node->DevComment.IsEmpty())
 			{
-				TooltipString.Append(FString::Printf(TEXT("\n\n%s"), *Node->DevComment));
+				TooltipString.Appendf(TEXT("\n\n%s"), *Node->DevComment);
 			}
 
 			// info related to conflicts
@@ -631,7 +631,7 @@ TSharedRef<ITableRow> SGameplayTagPicker::OnGenerateRow(TSharedPtr<FGameplayTagN
 			}
 		}
 
-		TooltipText = FText::FromString(TooltipString);
+		TooltipText = FText::FromString(MoveTemp(TooltipString));
 	}
 
 	TSharedPtr<SComboButton> ActionsCombo = SNew(SComboButton)
diff --git a/Engine/Plugins/Editor/GameplayTagsEditor/Source/GameplayTagsEditor/Private/SGameplayTagWidget.cpp b/Engine/Plugins/Editor/GameplayTagsEditor/Source/GameplayTagsEditor/Private/SGameplayTagWidget.cpp
--- a/Engine/Plugins/Editor/GameplayTagsEditor/Source/GameplayTagsEditor/Private/SGameplayTagWidget.cpp
+++ b/Engine/Plugins/Editor/GameplayTagsEditor/Source/GameplayTagsEditor/Private/SGameplayTagWidget.cpp
@@ -452,12 +452,12 @@ TSharedRef<ITableRow> SGameplayTagWidget::OnGenerateRow(TSharedPtr<FGameplayTagN
 			{
 				TagSource = FName(TEXT("Implicit"));
 			}
-			TooltipString.Append(FString::Printf(TEXT(" (%s)"), *TagSource.ToString()));
+			TooltipString.Appendf(TEXT(" (%s)"), *TagSource.ToString());
 
 			// tag comments
 			if (!Node->DevComment.IsEmpty())
 			{
-				TooltipString.Append(FString::Printf(TEXT("\n\n%s"), *Node->DevComment));
+				TooltipString.Appendf(TEXT("\n\n%s"), *Node->DevComment);
 			}
 
 			// info related to conflicts
@@ -477,7 +477,7 @@ TSharedRef<ITableRow> SGameplayTagWidget::OnGenerateRow(TSharedPtr<FGameplayTagN
 			}
 		}
 
-		TooltipText = FText::FromString(TooltipString);
+		TooltipText = FText::FromString(MoveTemp(TooltipString));
 	}
 
 	return SNew(STableRow< TSharedPtr<FGameplayTagNode> >, OwnerTable)
diff --git a/Engine/Source/Runtime/Core/Private/UObject/UnrealNames.cpp b/Engine/Source/Runtime/Core/Private/UObject/UnrealNames.cpp
--- a/Engine/Source/Runtime/Core/Private/UObject/UnrealNames.cpp
+++ b/Engine/Source/Runtime/Core/Private/UObject/UnrealNames.cpp
@@ -3944,10 +3944,10 @@ bool FName::IsValidXName(const FStringView& InName, const FString& InInvalidChar
 				{
 					if (FText::IsWhitespace(ch))
 					{
-						InvalidCharacters.Append(FString::Printf(TEXT(" U+%04X"), ch));
+						InvalidCharacters.Appendf(TEXT(" U+%04X"), ch);
 						continue;
 					}
-					InvalidCharacters.Append(FString::Printf(TEXT(" %c"), ch));
+					InvalidCharacters.Appendf(TEXT(" %c"), ch);
 				}
 				Args.Add(TEXT("InvalidCharacters"), FText::FromString(InvalidCharacters));
 
diff --git a/Engine/Source/Runtime/GameplayTags/Private/Tests/GameplayTagTests.cpp b/Engine/Source/Runtime/GameplayTags/Private/Tests/GameplayTagTests.cpp
--- a/Engine/Source/Runtime/GameplayTags/Private/Tests/GameplayTagTests.cpp
+++ b/Engine/Source/Runtime/GameplayTags/Private/Tests/GameplayTagTests.cpp
@@ -96,7 +96,7 @@ public:
 		FString CSV(TEXT(",Tag,CategoryText,"));
 		for (int32 Count = 0; Count < UE_ARRAY_COUNT(TestTags); Count++)
 		{
-			CSV.Append(FString::Printf(TEXT("\r\n%d,%.*s"), Count, TestTags[Count].Len(), TestTags[Count].GetData()));
+			CSV.Appendf(TEXT("\r\n%d,%.*s"), Count, TestTags[Count].Len(), TestTags[Count].GetData());
 		}
 
 		DataTable->CreateTableFromCSVString(CSV);
