GameplayTagQueryWidget: Do not dirty files on cancel with no change in a tag query

#rnx

diff --git a/Engine/Plugins/Editor/GameplayTagsEditor/Source/GameplayTagsEditor/Private/SGameplayTagQueryWidget.cpp b/Engine/Plugins/Editor/GameplayTagsEditor/Source/GameplayTagsEditor/Private/SGameplayTagQueryWidget.cpp
--- a/Engine/Plugins/Editor/GameplayTagsEditor/Source/GameplayTagsEditor/Private/SGameplayTagQueryWidget.cpp
+++ b/Engine/Plugins/Editor/GameplayTagsEditor/Source/GameplayTagsEditor/Private/SGameplayTagQueryWidget.cpp
@@ -201,9 +201,11 @@ FReply SGameplayTagQueryWidget::OnOkClicked()
 
 FReply SGameplayTagQueryWidget::OnCancelClicked()
 {
-	TagQueries = OriginalTagQueries;
-	
-	OnQueriesCommitted.ExecuteIfBound(TagQueries);	
+	if (TagQueries != OriginalTagQueries)
+	{
+		TagQueries = OriginalTagQueries;
+		OnQueriesCommitted.ExecuteIfBound(TagQueries);	
+	}
 	OnCancel.ExecuteIfBound();
 	
 	return FReply::Handled();
