Improved Gameplay Tags UX
- Added Gameplay Tag Manager entry to the Editor "Tools" menu. Now you don't have to look for a tag variable editor widget to start this from. Also dockable.
- Added "Search For References (all children)" item to the context menu for individual tags. Triggers the Reference Viewer with the selected tag AND all its recursive children.

diff --git a/Engine/Plugins/Editor/GameplayTagsEditor/Source/GameplayTagsEditor/GameplayTagsEditor.Build.cs b/Engine/Plugins/Editor/GameplayTagsEditor/Source/GameplayTagsEditor/GameplayTagsEditor.Build.cs
--- a/Engine/Plugins/Editor/GameplayTagsEditor/Source/GameplayTagsEditor/GameplayTagsEditor.Build.cs
+++ b/Engine/Plugins/Editor/GameplayTagsEditor/Source/GameplayTagsEditor/GameplayTagsEditor.Build.cs
@@ -6,6 +6,12 @@ namespace UnrealBuildTool.Rules
 	{
 		public GameplayTagsEditor(ReadOnlyTargetRules Target) : base(Target)
 		{
+			PrivateIncludePathModuleNames.AddRange(
+				new string[] {
+					"WorkspaceMenuStructure",
+				}
+			);
+
 			PublicIncludePathModuleNames.AddRange(
 				new string[] {
 					"AssetTools",
diff --git a/Engine/Plugins/Editor/GameplayTagsEditor/Source/GameplayTagsEditor/Private/GameplayTagsEditorModule.cpp b/Engine/Plugins/Editor/GameplayTagsEditor/Source/GameplayTagsEditor/Private/GameplayTagsEditorModule.cpp
--- a/Engine/Plugins/Editor/GameplayTagsEditor/Source/GameplayTagsEditor/Private/GameplayTagsEditorModule.cpp
+++ b/Engine/Plugins/Editor/GameplayTagsEditor/Source/GameplayTagsEditor/Private/GameplayTagsEditorModule.cpp
@@ -30,10 +30,19 @@
 #include "UObject/ObjectSaveContext.h"
 #include "GameplayTagStyle.h"
 #include "SGameplayTagPicker.h"
+#include "Framework/Docking/TabManager.h"
+#include "WorkspaceMenuStructure.h"
+#include "WorkspaceMenuStructureModule.h"
+#include "Widgets/Docking/SDockTab.h"
 
 #define LOCTEXT_NAMESPACE "GameplayTagEditor"
 
 
+namespace GameplayTagEditorModule
+{
+	static const FName GameplayTagManagerApp = FName("GameplayTagManagerApp");
+}
+
 class FGameplayTagsEditorModule
 	: public IGameplayTagsEditorModule
 {
@@ -43,6 +52,25 @@ public:
 
 	virtual void StartupModule() override
 	{
+
+		auto OnCreateLambda = [](const FSpawnTabArgs& Args)
+			{
+				TSharedRef<SDockTab> DockTab = SNew(SDockTab)
+					.Clipping(EWidgetClipping::ClipToBounds)
+					.Label(LOCTEXT("GameplayTagPicker_ManagerTitle", "Gameplay Tag Manager"));
+
+				FGameplayTagManagerWindowArgs TagManagerWindowArgs;
+				TagManagerWindowArgs.bRestrictedTags = false;
+				DockTab->SetContent( UE::GameplayTags::Editor::Create(TagManagerWindowArgs) );
+				return DockTab;
+			};
+		
+		FGlobalTabmanager::Get()->RegisterNomadTabSpawner(GameplayTagEditorModule::GameplayTagManagerApp, FOnSpawnTab::CreateLambda(OnCreateLambda))
+			.SetDisplayName(LOCTEXT("GameplayTagPicker_ManagerTitle", "Gameplay Tag Manager"))
+			.SetTooltipText(LOCTEXT("GameplayTagPicker_ManagerTitle", "Gameplay Tag Manager"))
+			.SetGroup(WorkspaceMenu::GetMenuStructure().GetToolsCategory())
+			.SetIcon(FSlateIcon(FAppStyle::GetAppStyleSetName(), "ClassIcon.UserDefinedStruct"));
+
 		FCoreDelegates::OnPostEngineInit.AddRaw(this, &FGameplayTagsEditorModule::OnPostEngineInit);
 		FGameplayTagStyle::Initialize();
 	}
diff --git a/Engine/Plugins/Editor/GameplayTagsEditor/Source/GameplayTagsEditor/Private/SGameplayTagPicker.cpp b/Engine/Plugins/Editor/GameplayTagsEditor/Source/GameplayTagsEditor/Private/SGameplayTagPicker.cpp
--- a/Engine/Plugins/Editor/GameplayTagsEditor/Source/GameplayTagsEditor/Private/SGameplayTagPicker.cpp
+++ b/Engine/Plugins/Editor/GameplayTagsEditor/Source/GameplayTagsEditor/Private/SGameplayTagPicker.cpp
@@ -1191,6 +1191,15 @@ TSharedRef<SWidget> SGameplayTagPicker::MakeTagActionsMenu(TSharedPtr<FGameplayT
 			FUIAction(FExecuteAction::CreateSP(this, &SGameplayTagPicker::OnSearchForReferences, InTagNode, ActionsCombo)));
 	}
 
+	// Search for References of all children
+	if (FEditorDelegates::OnOpenReferenceViewer.IsBound())
+	{
+		MenuBuilder.AddMenuEntry(LOCTEXT("GameplayTagPicker_SearchForReferences", "Search For References (all children)"),
+			FText::Format(LOCTEXT("GameplayTagPicker_SearchForReferencesTooltip", "Find references to the tag {0} and all its children"), FText::AsCultureInvariant(InTagNode->GetCompleteTagString())),
+			FSlateIcon(FAppStyle::GetAppStyleSetName(), "Icons.Search"),
+			FUIAction(FExecuteAction::CreateSP(this, &SGameplayTagPicker::OnSearchForReferencesAsChildren, InTagNode, ActionsCombo)));
+	}
+
 	// Copy Name to Clipboard
 	MenuBuilder.AddMenuEntry(LOCTEXT("GameplayTagPicker_CopyNameToClipboard", "Copy Name to Clipboard"),
 	FText::Format(LOCTEXT("GameplayTagPicker_CopyNameToClipboardTooltip", "Copy tag {0} to clipboard"), FText::AsCultureInvariant(InTagNode->GetCompleteTagString())),
@@ -1336,6 +1345,32 @@ void SGameplayTagPicker::OnSearchForReferences(TSharedPtr<FGameplayTagNode> InTa
 	}
 }
 
+void SGameplayTagPicker::OnSearchForReferencesAsChildren(TSharedPtr<FGameplayTagNode> InTagNode, TSharedPtr<SComboButton> OwnerCombo)
+{
+	if (InTagNode.IsValid())
+	{
+		UGameplayTagsManager& TagManager = UGameplayTagsManager::Get();
+		UScriptStruct* GameplayTagStruct = FGameplayTag::StaticStruct();
+
+		TArray<FAssetIdentifier> AssetIdentifiers;
+		AssetIdentifiers.Add(FAssetIdentifier(GameplayTagStruct, InTagNode->GetCompleteTagName()));
+
+		FGameplayTagContainer TagContainer = TagManager.RequestGameplayTagChildrenInDictionary(InTagNode->GetCompleteTag());
+		for (FGameplayTag ChildTag : TagContainer.GetGameplayTagArray())
+		{
+			FAssetIdentifier ChildAssetIdentifier(GameplayTagStruct, ChildTag.GetTagName());
+			AssetIdentifiers.Add(ChildAssetIdentifier);
+		}
+
+		FEditorDelegates::OnOpenReferenceViewer.Broadcast(AssetIdentifiers, FReferenceViewerParams());
+	}
+
+	if (OwnerCombo.IsValid())
+	{
+		OwnerCombo->SetIsOpen(false);
+	}
+}
+
 void SGameplayTagPicker::OnCopyTagNameToClipboard(TSharedPtr<FGameplayTagNode> InTagNode, TSharedPtr<SComboButton> OwnerCombo)
 {
 	if (InTagNode.IsValid())
@@ -1786,6 +1821,20 @@ void CloseGameplayTagWindow(TWeakPtr<SGameplayTagPicker> TagWidget)
 	GlobalTagWidget = nullptr;
 }
 
+TSharedRef<SGameplayTagPicker> Create(const FGameplayTagManagerWindowArgs& Args)
+{
+	return SNew(SGameplayTagPicker)
+		.Filter(Args.Filter)
+		.ReadOnly(false)
+		.MaxHeight(0.0f) // unbounded
+		.MultiSelect(false)
+		.SettingsName(TEXT("Manager"))
+		.GameplayTagPickerMode(EGameplayTagPickerMode::ManagementMode)
+		.RestrictedTags(Args.bRestrictedTags)
+		;
+}
+
+
 TWeakPtr<SGameplayTagPicker> OpenGameplayTagManager(const FGameplayTagManagerWindowArgs& Args)
 {
 	TSharedPtr<SWindow> GameplayTagPickerWindow = GlobalTagWidgetWindow.Pin();
@@ -1797,17 +1846,8 @@ TWeakPtr<SGameplayTagPicker> OpenGameplayTagManager(const FGameplayTagManagerWin
 		// Close all other GameplayTag windows.
 		CloseGameplayTagWindow(nullptr);
 
-		const FVector2D WindowSize(800, 800);
-		
-		TagWidget = SNew(SGameplayTagPicker)
-			.Filter(Args.Filter)
-			.ReadOnly(false)
-			.MaxHeight(0.0f) // unbounded
-			.MultiSelect(false)
-			.SettingsName(TEXT("Manager"))
-			.GameplayTagPickerMode(EGameplayTagPickerMode::ManagementMode)
-			.RestrictedTags(Args.bRestrictedTags)
-		;
+
+		TagWidget = Create(Args);
 
 		FText Title = Args.Title;
 		if (Title.IsEmpty())
@@ -1821,7 +1861,8 @@ TWeakPtr<SGameplayTagPicker> OpenGameplayTagManager(const FGameplayTagManagerWin
 				Title = LOCTEXT("GameplayTagPicker_ManagerTitle", "Gameplay Tag Manager");
 			}
 		}
-		
+
+		const FVector2D WindowSize(800, 800);
 		GameplayTagPickerWindow = SNew(SWindow)
 			.Title(Title)
 			.HasCloseButton(true)
@@ -1842,8 +1883,6 @@ TWeakPtr<SGameplayTagPicker> OpenGameplayTagManager(const FGameplayTagManagerWin
 					]
 				]
 			];
-
-		TWeakPtr<SGameplayTagPicker> WeakTagWidget = TagWidget;
 		
 		// NOTE: FGlobalTabmanager::Get()-> is actually dereferencing a SharedReference, not a SharedPtr, so it cannot be null.
 		if (FGlobalTabmanager::Get()->GetRootWindow().IsValid())
diff --git a/Engine/Plugins/Editor/GameplayTagsEditor/Source/GameplayTagsEditor/Public/SGameplayTagPicker.h b/Engine/Plugins/Editor/GameplayTagsEditor/Source/GameplayTagsEditor/Public/SGameplayTagPicker.h
--- a/Engine/Plugins/Editor/GameplayTagsEditor/Source/GameplayTagsEditor/Public/SGameplayTagPicker.h
+++ b/Engine/Plugins/Editor/GameplayTagsEditor/Source/GameplayTagsEditor/Public/SGameplayTagPicker.h
@@ -350,6 +350,9 @@ private:
 	/** Searches for all references for the selected tag */
 	void OnSearchForReferences(TSharedPtr<FGameplayTagNode> InTagNode, TSharedPtr<SComboButton> OwnerCombo);
 
+	/** Searches for all references for the selected tag and its children */
+	void OnSearchForReferencesAsChildren(TSharedPtr<FGameplayTagNode> InTagNode, TSharedPtr<SComboButton> OwnerCombo);
+
 	/** Copies individual tag's name to clipboard. */
 	void OnCopyTagNameToClipboard(TSharedPtr<FGameplayTagNode> InTagNode, TSharedPtr<SComboButton> OwnerCombo);
 	
@@ -464,5 +467,5 @@ struct FGameplayTagManagerWindowArgs
 namespace UE::GameplayTags::Editor
 {
 	GAMEPLAYTAGSEDITOR_API TWeakPtr<SGameplayTagPicker> OpenGameplayTagManager(const FGameplayTagManagerWindowArgs& Args);
-	GAMEPLAYTAGSEDITOR_API TSharedRef<SWidget> Create(const FGameplayTagManagerWindowArgs& Args);
+	GAMEPLAYTAGSEDITOR_API TSharedRef<SGameplayTagPicker> Create(const FGameplayTagManagerWindowArgs& Args);
 };
\ No newline at end of file
