Fixing a localization break

diff --git a/Engine/Plugins/Editor/GameplayTagsEditor/Source/GameplayTagsEditor/Private/SGameplayTagPicker.cpp b/Engine/Plugins/Editor/GameplayTagsEditor/Source/GameplayTagsEditor/Private/SGameplayTagPicker.cpp
--- a/Engine/Plugins/Editor/GameplayTagsEditor/Source/GameplayTagsEditor/Private/SGameplayTagPicker.cpp
+++ b/Engine/Plugins/Editor/GameplayTagsEditor/Source/GameplayTagsEditor/Private/SGameplayTagPicker.cpp
@@ -1194,8 +1194,8 @@ TSharedRef<SWidget> SGameplayTagPicker::MakeTagActionsMenu(TSharedPtr<FGameplayT
 	// Search for References of all children
 	if (FEditorDelegates::OnOpenReferenceViewer.IsBound())
 	{
-		MenuBuilder.AddMenuEntry(LOCTEXT("GameplayTagPicker_SearchForReferences", "Search For References (all children)"),
-			FText::Format(LOCTEXT("GameplayTagPicker_SearchForReferencesTooltip", "Find references to the tag {0} and all its children"), FText::AsCultureInvariant(InTagNode->GetCompleteTagString())),
+		MenuBuilder.AddMenuEntry(LOCTEXT("GameplayTagPicker_SearchForReferencesAsChildren", "Search For References (all children)"),
+			FText::Format(LOCTEXT("GameplayTagPicker_SearchForReferencesAsChildrenTooltip", "Find references to the tag {0} and all its children"), FText::AsCultureInvariant(InTagNode->GetCompleteTagString())),
 			FSlateIcon(FAppStyle::GetAppStyleSetName(), "Icons.Search"),
 			FUIAction(FExecuteAction::CreateSP(this, &SGameplayTagPicker::OnSearchForReferencesAsChildren, InTagNode, ActionsCombo)));
 	}
