GameplayTags and GameplayTagContainers allow the user to edit the "Categories" metadata for tag filtering in the owning BP.

diff --git a/Engine/Plugins/Editor/GameplayTagsEditor/Source/GameplayTagsEditor/Private/GameplayTagContainerCustomization.cpp b/Engine/Plugins/Editor/GameplayTagsEditor/Source/GameplayTagsEditor/Private/GameplayTagContainerCustomization.cpp
--- a/Engine/Plugins/Editor/GameplayTagsEditor/Source/GameplayTagsEditor/Private/GameplayTagContainerCustomization.cpp
+++ b/Engine/Plugins/Editor/GameplayTagsEditor/Source/GameplayTagsEditor/Private/GameplayTagContainerCustomization.cpp
@@ -2,14 +2,16 @@
 
 #include "GameplayTagContainerCustomization.h"
 #include "DetailWidgetRow.h"
+#include "Engine/BlueprintGeneratedClass.h"
 #include "GameplayTagsEditorModule.h"
 #include "GameplayTagsManager.h"
-#include "Widgets/Input/SHyperlink.h"
 #include "SGameplayTagContainerCombo.h"
 #include "HAL/PlatformApplicationMisc.h"
 #include "ScopedTransaction.h"
 #include "GameplayTagEditorUtilities.h"
+#include "IDetailChildrenBuilder.h"
 #include "SGameplayTagPicker.h"
+#include "Kismet2/BlueprintEditorUtils.h"
 
 #define LOCTEXT_NAMESPACE "GameplayTagContainerCustomization"
 
@@ -25,17 +27,48 @@ TSharedRef<IPropertyTypeCustomization> FGameplayTagContainerCustomizationPublic:
 }
 
 FGameplayTagContainerCustomization::FGameplayTagContainerCustomization()
+	: CreatorBP(nullptr)
 {
 }
 
 void FGameplayTagContainerCustomization::CustomizeHeader(TSharedRef<IPropertyHandle> InStructPropertyHandle, FDetailWidgetRow& HeaderRow, IPropertyTypeCustomizationUtils& StructCustomizationUtils)
 {
 	StructPropertyHandle = InStructPropertyHandle;
+	FString CurrentValue;
+	StructPropertyHandle->GetValueAsFormattedString(CurrentValue);
+
+	if (StructPropertyHandle->HasMetaData("Categories"))
+	{
+		TArray<FString> CategoriesStringList;
+		const FString CategoriesString = UGameplayTagsManager::Get().GetCategoriesMetaFromPropertyHandle(StructPropertyHandle);
+
+		if (CategoriesString.ParseIntoArray(CategoriesStringList, TEXT(","), true) > 0)
+		{
+			UGameplayTagsManager::Get().RequestGameplayTagContainer(CategoriesStringList, MetaDataCategories);
+		}
+	}
+	
+	FGameplayTagContainer CurrentTagContainer;
+	CurrentTagContainer.FromExportString(CurrentValue);
+	bool bHasValidCategories = CurrentTagContainer.HasAll(MetaDataCategories);
 
 	HeaderRow
 		.NameContent()
 		[
-			StructPropertyHandle->CreatePropertyNameWidget()
+			SNew(SHorizontalBox)
+			+ SHorizontalBox::Slot()
+			.Padding(FMargin(0,2,0,1))
+			.AutoWidth()
+			[
+				SNew(SImage)
+				.Visibility(bHasValidCategories ? EVisibility::Collapsed : EVisibility::Visible)
+				.Image(FAppStyle::GetBrush("Icons.WarningWithColor"))
+				.ToolTipText(FText::FromString("Current tag values do not match the Gameplay Tag Roots filter."))
+			]
+			+ SHorizontalBox::Slot()
+			[
+				StructPropertyHandle->CreatePropertyNameWidget()
+			]
 		]
 		.ValueContent()
 		.VAlign(VAlign_Center)
@@ -52,6 +85,69 @@ void FGameplayTagContainerCustomization::CustomizeHeader(TSharedRef<IPropertyHan
 		FCanExecuteAction::CreateSP(this, &FGameplayTagContainerCustomization::CanPasteTag)));
 }
 
+void FGameplayTagContainerCustomization::CustomizeChildren(TSharedRef<IPropertyHandle> InStructPropertyHandle, IDetailChildrenBuilder& ChildBuilder,
+	IPropertyTypeCustomizationUtils& StructCustomizationUtils)
+{
+	if (StructPropertyHandle->GetProperty()->IsNative())
+	{
+		return;
+	}
+
+	if (const UBlueprintGeneratedClass* GeneratedClass = Cast<UBlueprintGeneratedClass>(StructPropertyHandle->GetOuterBaseClass()))
+	{
+		CreatorBP = Cast<UBlueprint>(GeneratedClass->ClassGeneratedBy);
+		if (!FBlueprintEditorUtils::IsVariableCreatedByBlueprint(CreatorBP.Get(), StructPropertyHandle->GetProperty()))
+		{
+			return;
+		}
+	}
+
+	ChildBuilder.AddCustomRow(LOCTEXT("Gameplay Tag Roots", "Categories"))
+			.NameContent()
+			[
+				SNew(STextBlock)
+					.Text(FText::FromString("Gameplay Tag Roots"))
+					.ToolTipText(FText::FromString("Selects the allowed root tags this variable can inherit from"))
+					.Font(StructCustomizationUtils.GetRegularFont())
+			]
+			.ValueContent()
+			.VAlign(VAlign_Center)
+			[
+				SNew(SGameplayTagContainerCombo)
+					.TagContainer(this, &FGameplayTagContainerCustomization::GetMetaDataCategories)
+					.ToolTipText(FText::FromString("Selects the allowed root tags this variable can inherit from"))
+					.OnTagContainerChanged_Lambda([this](const FGameplayTagContainer& NewTagContainer)
+					{
+						MetaDataCategories = NewTagContainer;
+					})
+					.OnTagCleared(this, &FGameplayTagContainerCustomization::OnTagCleared)
+					.OnTagContainerComboClosed(this, &FGameplayTagContainerCustomization::OnGameplayTagContainerComboClosed)
+					
+			];
+}
+
+void FGameplayTagContainerCustomization::OnGameplayTagContainerComboClosed(const FGameplayTagContainer& NewTagContainer)
+{
+	MetaDataCategories = NewTagContainer;
+	if (!MetaDataCategories.IsEmpty())
+	{
+		FString CurrentTag;
+		FString ContainerString = MetaDataCategories.ToStringSimple();
+		ContainerString.ReplaceInline(TEXT(" "), TEXT(""));
+		StructPropertyHandle->GetProperty()->SetMetaData(TEXT("Categories"), *ContainerString);
+		FBlueprintEditorUtils::SetBlueprintVariableMetaData(CreatorBP.Get(), StructPropertyHandle->GetProperty()->GetFName(), StructPropertyHandle->GetProperty()->GetOwner<UFunction>(), TEXT("Categories"), ContainerString);
+	}
+	else
+	{
+		FBlueprintEditorUtils::RemoveBlueprintVariableMetaData(CreatorBP.Get(), StructPropertyHandle->GetProperty()->GetFName(), StructPropertyHandle->GetProperty()->GetOwner<UFunction>(), TEXT("Categories"));
+	}
+}
+
+void FGameplayTagContainerCustomization::OnTagCleared()
+{
+	OnGameplayTagContainerComboClosed(MetaDataCategories);
+}
+
 void FGameplayTagContainerCustomization::OnPasteTag() const
 {
 	if (!StructPropertyHandle.IsValid())
diff --git a/Engine/Plugins/Editor/GameplayTagsEditor/Source/GameplayTagsEditor/Private/GameplayTagContainerCustomization.h b/Engine/Plugins/Editor/GameplayTagsEditor/Source/GameplayTagsEditor/Private/GameplayTagContainerCustomization.h
--- a/Engine/Plugins/Editor/GameplayTagsEditor/Source/GameplayTagsEditor/Private/GameplayTagContainerCustomization.h
+++ b/Engine/Plugins/Editor/GameplayTagsEditor/Source/GameplayTagsEditor/Private/GameplayTagContainerCustomization.h
@@ -3,6 +3,7 @@
 #pragma once
 
 #include "IPropertyTypeCustomization.h"
+#include "GameplayTagContainer.h"
 
 struct FGameplayTag;
 class IPropertyHandle;
@@ -20,11 +21,22 @@ public:
 	virtual void CustomizeHeader(TSharedRef<IPropertyHandle> StructPropertyHandle, FDetailWidgetRow& HeaderRow, IPropertyTypeCustomizationUtils& StructCustomizationUtils) override;
 	
 	/** Overridden to do nothing */
-	virtual void CustomizeChildren(TSharedRef<IPropertyHandle> InStructPropertyHandle, IDetailChildrenBuilder& ChildBuilder, IPropertyTypeCustomizationUtils& StructCustomizationUtils) override {}
+	virtual void CustomizeChildren(TSharedRef<IPropertyHandle> InStructPropertyHandle, IDetailChildrenBuilder& ChildBuilder, IPropertyTypeCustomizationUtils& StructCustomizationUtils) override;
 
 private:
+
+	FGameplayTagContainer GetMetaDataCategories() const
+	{
+		return MetaDataCategories;
+	}
+
+	void OnGameplayTagContainerComboClosed(const FGameplayTagContainer& NewTagContainer);
+	void OnTagCleared();
+	
 	/** Cached property handle */
 	TSharedPtr<IPropertyHandle> StructPropertyHandle;
+	FGameplayTagContainer MetaDataCategories;
+	TWeakObjectPtr<UBlueprint> CreatorBP;
 
 	void OnPasteTag() const;
 	bool CanPasteTag() const;
diff --git a/Engine/Plugins/Editor/GameplayTagsEditor/Source/GameplayTagsEditor/Private/GameplayTagCustomization.cpp b/Engine/Plugins/Editor/GameplayTagsEditor/Source/GameplayTagsEditor/Private/GameplayTagCustomization.cpp
--- a/Engine/Plugins/Editor/GameplayTagsEditor/Source/GameplayTagsEditor/Private/GameplayTagCustomization.cpp
+++ b/Engine/Plugins/Editor/GameplayTagsEditor/Source/GameplayTagsEditor/Private/GameplayTagCustomization.cpp
@@ -1,14 +1,16 @@
 // Copyright Epic Games, Inc. All Rights Reserved.
 
 #include "GameplayTagCustomization.h"
+#include "Engine/BlueprintGeneratedClass.h"
 #include "PropertyHandle.h"
 #include "DetailWidgetRow.h"
 #include "GameplayTagsManager.h"
 #include "IDetailChildrenBuilder.h"
 #include "GameplayTagsEditorModule.h"
-#include "Widgets/Input/SHyperlink.h"
 #include "SGameplayTagCombo.h"
+#include "SGameplayTagContainerCombo.h"
 #include "SGameplayTagPicker.h"
+#include "Kismet2/BlueprintEditorUtils.h"
 
 #define LOCTEXT_NAMESPACE "GameplayTagCustomization"
 
@@ -32,17 +34,48 @@ TSharedRef<IPropertyTypeCustomization> FGameplayTagCustomizationPublic::MakeInst
 //---------------------------------------------------------
 
 FGameplayTagCustomization::FGameplayTagCustomization()
+	: CreatorBP(nullptr)
 {
 }
 
 void FGameplayTagCustomization::CustomizeHeader(TSharedRef<IPropertyHandle> InStructPropertyHandle, FDetailWidgetRow& HeaderRow, IPropertyTypeCustomizationUtils& StructCustomizationUtils)
 {
 	StructPropertyHandle = InStructPropertyHandle;
+	FString CurrentValue;
+	StructPropertyHandle->GetValueAsFormattedString(CurrentValue);
+
+	if (StructPropertyHandle->HasMetaData("Categories"))
+	{
+		TArray<FString> CategoriesStringList;
+		const FString CategoriesString = UGameplayTagsManager::Get().GetCategoriesMetaFromPropertyHandle(StructPropertyHandle);
+
+		if (CategoriesString.ParseIntoArray(CategoriesStringList, TEXT(","), true) > 0)
+		{
+			UGameplayTagsManager::Get().RequestGameplayTagContainer(CategoriesStringList, MetaDataCategories);
+		}
+	}
+
+	FGameplayTag CurrentTag;
+	CurrentTag.FromExportString(CurrentValue);
+	bool bHasValidCategories = CurrentTag == FGameplayTag::EmptyTag || MetaDataCategories.IsEmpty() || CurrentTag.MatchesAny(MetaDataCategories);
 
 	HeaderRow
 	.NameContent()
 	[
-		StructPropertyHandle->CreatePropertyNameWidget()
+		SNew(SHorizontalBox)
+		+ SHorizontalBox::Slot()
+		.Padding(FMargin(0,2,0,1))
+		.AutoWidth()
+		[
+			SNew(SImage)
+			.Visibility(bHasValidCategories ? EVisibility::Collapsed : EVisibility::Visible)
+			.Image(FAppStyle::GetBrush("Icons.WarningWithColor"))
+			.ToolTipText(FText::FromString("Current tag values do not match the Gameplay Tag Roots filter."))
+		]
+		+ SHorizontalBox::Slot()
+		[
+			StructPropertyHandle->CreatePropertyNameWidget()
+		]
 	]
 	.ValueContent()
 	[
@@ -55,6 +88,82 @@ void FGameplayTagCustomization::CustomizeHeader(TSharedRef<IPropertyHandle> InSt
 	];
 }
 
+void FGameplayTagCustomization::CustomizeChildren(TSharedRef<IPropertyHandle> InStructPropertyHandle, IDetailChildrenBuilder& ChildBuilder,
+	IPropertyTypeCustomizationUtils& StructCustomizationUtils)
+{
+	if (StructPropertyHandle->GetProperty()->IsNative())
+	{
+		return;
+	}
+
+	if (const UBlueprintGeneratedClass* GeneratedClass = Cast<UBlueprintGeneratedClass>(StructPropertyHandle->GetOuterBaseClass()))
+	{
+		CreatorBP = Cast<UBlueprint>(GeneratedClass->ClassGeneratedBy);
+		if (!FBlueprintEditorUtils::IsVariableCreatedByBlueprint(CreatorBP.Get(), StructPropertyHandle->GetProperty()))
+		{
+			return;
+		}
+	}
+
+	if (StructPropertyHandle->HasMetaData("Categories"))
+	{
+		TArray<FString> CategoriesStringList;
+		const FString CategoriesString = UGameplayTagsManager::Get().GetCategoriesMetaFromPropertyHandle(StructPropertyHandle);
+
+		if (CategoriesString.ParseIntoArray(CategoriesStringList, TEXT(","), true) > 0)
+		{
+			UGameplayTagsManager::Get().RequestGameplayTagContainer(CategoriesStringList, MetaDataCategories);
+		}
+	}
+
+	ChildBuilder.AddCustomRow(LOCTEXT("Gameplay Tag Roots", "Categories"))
+			.NameContent()
+			[
+				SNew(STextBlock)
+					.Text(FText::FromString("Gameplay Tag Roots"))
+					.ToolTipText(FText::FromString("Selects the allowed root tags this variable can inherit from"))
+					.Font(StructCustomizationUtils.GetRegularFont())
+			]
+			.ValueContent()
+			.VAlign(VAlign_Center)
+			[
+				SNew(SGameplayTagContainerCombo)
+					.TagContainer(this, &FGameplayTagCustomization::GetMetaDataCategories)
+					.ToolTipText(FText::FromString("Selects the allowed root tags this variable can inherit from"))
+					.OnTagContainerChanged_Lambda([this](const FGameplayTagContainer& NewTagContainer)
+					{
+						MetaDataCategories = NewTagContainer;
+					})
+					.OnTagCleared(this, &FGameplayTagCustomization::OnTagCleared)
+					.OnTagContainerComboClosed(this, &FGameplayTagCustomization::OnGameplayTagContainerComboClosed)
+					
+			];
+}
+
+void FGameplayTagCustomization::OnGameplayTagContainerComboClosed(const FGameplayTagContainer& NewTagContainer)
+{
+	MetaDataCategories = NewTagContainer;
+	if (!MetaDataCategories.IsEmpty())
+	{
+		FString CurrentTag;
+		FString ContainerString = MetaDataCategories.ToStringSimple();
+		ContainerString.ReplaceInline(TEXT(" "), TEXT(""));
+		StructPropertyHandle->GetProperty()->SetMetaData(TEXT("Categories"), *ContainerString);
+		FBlueprintEditorUtils::SetBlueprintVariableMetaData(CreatorBP.Get(), StructPropertyHandle->GetProperty()->GetFName(), 
+		StructPropertyHandle->GetProperty()->GetOwner<UFunction>(), TEXT("Categories"), ContainerString);
+	}
+	else
+	{
+		FBlueprintEditorUtils::RemoveBlueprintVariableMetaData(CreatorBP.Get(), StructPropertyHandle->GetProperty()->GetFName(), 
+		StructPropertyHandle->GetProperty()->GetOwner<UFunction>(), TEXT("Categories"));
+	}
+}
+
+void FGameplayTagCustomization::OnTagCleared()
+{
+	OnGameplayTagContainerComboClosed(MetaDataCategories);
+}
+
 //---------------------------------------------------------
 // FGameplayTagCreationWidgetHelperDetails
 //---------------------------------------------------------
diff --git a/Engine/Plugins/Editor/GameplayTagsEditor/Source/GameplayTagsEditor/Private/GameplayTagCustomization.h b/Engine/Plugins/Editor/GameplayTagsEditor/Source/GameplayTagsEditor/Private/GameplayTagCustomization.h
--- a/Engine/Plugins/Editor/GameplayTagsEditor/Source/GameplayTagsEditor/Private/GameplayTagCustomization.h
+++ b/Engine/Plugins/Editor/GameplayTagsEditor/Source/GameplayTagsEditor/Private/GameplayTagCustomization.h
@@ -21,17 +21,27 @@ public:
 
 	/** Overridden to show an edit button to launch the gameplay tag editor */
 	virtual void CustomizeHeader(TSharedRef<IPropertyHandle> StructPropertyHandle, FDetailWidgetRow& HeaderRow, IPropertyTypeCustomizationUtils& StructCustomizationUtils) override;
-	
+
 	/** Overridden to do nothing */
-	virtual void CustomizeChildren(TSharedRef<IPropertyHandle> InStructPropertyHandle, IDetailChildrenBuilder& ChildBuilder, IPropertyTypeCustomizationUtils& StructCustomizationUtils) override {}
+	virtual void CustomizeChildren(TSharedRef<IPropertyHandle> InStructPropertyHandle, IDetailChildrenBuilder& ChildBuilder, IPropertyTypeCustomizationUtils& StructCustomizationUtils) override;
 
 private:
-
+	
 	/** Cached property handle */
 	TSharedPtr<IPropertyHandle> StructPropertyHandle;
 
+	FGameplayTagContainer GetMetaDataCategories() const
+	{
+			return MetaDataCategories;
+	}
+
+	void OnGameplayTagContainerComboClosed(const FGameplayTagContainer& NewTagContainer);
+	void OnTagCleared();
+
 	/** Edited Tag */
 	FGameplayTag Tag;
+	FGameplayTagContainer MetaDataCategories;
+	TWeakObjectPtr<UBlueprint> CreatorBP;
 };
 
 /** Customization for FGameplayTagCreationWidgetHelper showing an add tag button */
diff --git a/Engine/Plugins/Editor/GameplayTagsEditor/Source/GameplayTagsEditor/Private/SGameplayTagContainerCombo.cpp b/Engine/Plugins/Editor/GameplayTagsEditor/Source/GameplayTagsEditor/Private/SGameplayTagContainerCombo.cpp
--- a/Engine/Plugins/Editor/GameplayTagsEditor/Source/GameplayTagsEditor/Private/SGameplayTagContainerCombo.cpp
+++ b/Engine/Plugins/Editor/GameplayTagsEditor/Source/GameplayTagsEditor/Private/SGameplayTagContainerCombo.cpp
@@ -40,6 +40,8 @@ void SGameplayTagContainerCombo::Construct(const FArguments& InArgs)
 	Filter = InArgs._Filter;
 	SettingsName = InArgs._SettingsName;
 	bIsReadOnly = InArgs._ReadOnly;
+	OnTagCleared = InArgs._OnTagCleared;
+	OnTagContainerComboClosed = InArgs._OnTagContainerComboClosed;
 	OnTagContainerChanged = InArgs._OnTagContainerChanged;
 	PropertyHandle = InArgs._PropertyHandle;
 
@@ -194,6 +196,10 @@ void SGameplayTagContainerCombo::OnMenuOpenChanged(const bool bOpen)
 		}
 		TagPicker->RequestScrollToView(TagToHilight);
 	}
+	else
+	{
+		OnTagContainerComboClosed.ExecuteIfBound(CachedTagContainers[0]);
+	}
 	// Reset tag to hilight
 	TagToHilight = FGameplayTag();
 }
@@ -536,6 +542,7 @@ FReply SGameplayTagContainerCombo::OnClearTagClicked(const FGameplayTag TagToCle
 	}
 
 	RefreshTagContainers();
+	OnTagCleared.ExecuteIfBound();
 
 	return FReply::Handled();
 }
diff --git a/Engine/Plugins/Editor/GameplayTagsEditor/Source/GameplayTagsEditor/Public/SGameplayTagContainerCombo.h b/Engine/Plugins/Editor/GameplayTagsEditor/Source/GameplayTagsEditor/Public/SGameplayTagContainerCombo.h
--- a/Engine/Plugins/Editor/GameplayTagsEditor/Source/GameplayTagsEditor/Public/SGameplayTagContainerCombo.h
+++ b/Engine/Plugins/Editor/GameplayTagsEditor/Source/GameplayTagsEditor/Public/SGameplayTagContainerCombo.h
@@ -26,6 +26,8 @@ class SGameplayTagContainerCombo : public SCompoundWidget
 	
 public:
 
+	DECLARE_DELEGATE(FOnTagCleared)
+	DECLARE_DELEGATE_OneParam(FOnTagContainerComboClosed, const FGameplayTagContainer& /*TagContainer*/)
 	DECLARE_DELEGATE_OneParam(FOnTagContainerChanged, const FGameplayTagContainer& /*TagContainer*/)
 
 	SLATE_BEGIN_ARGS(SGameplayTagContainerCombo)
@@ -58,6 +60,12 @@ public:
 		// Callback for when button body is pressed with RMB
 		SLATE_EVENT(SGameplayTagChip::FOnMenu, OnMenu)
 
+		// Called when the tag container combo selection is closed
+		SLATE_EVENT(FOnTagCleared, OnTagCleared)
+
+		// Called when the tag container combo selection is closed
+		SLATE_EVENT(FOnTagContainerComboClosed, OnTagContainerComboClosed)
+
 		// Called when a tag container changes
 		SLATE_EVENT(FOnTagContainerChanged, OnTagContainerChanged)
 	SLATE_END_ARGS();
@@ -111,6 +119,8 @@ private:
 	FString Filter;
 	FString SettingsName;
 	bool bIsReadOnly = false;
+	FOnTagCleared OnTagCleared;
+	FOnTagContainerComboClosed OnTagContainerComboClosed;
 	FOnTagContainerChanged OnTagContainerChanged;
 	TSharedPtr<IPropertyHandle> PropertyHandle;
 	TSharedPtr<SMenuAnchor> MenuAnchor;
