GameplayTag and GameplayTagContainer customization should not error if the metadata tag isn't found.

diff --git a/Engine/Plugins/Editor/GameplayTagsEditor/Source/GameplayTagsEditor/Private/GameplayTagContainerCustomization.cpp b/Engine/Plugins/Editor/GameplayTagsEditor/Source/GameplayTagsEditor/Private/GameplayTagContainerCustomization.cpp
--- a/Engine/Plugins/Editor/GameplayTagsEditor/Source/GameplayTagsEditor/Private/GameplayTagContainerCustomization.cpp
+++ b/Engine/Plugins/Editor/GameplayTagsEditor/Source/GameplayTagsEditor/Private/GameplayTagContainerCustomization.cpp
@@ -44,7 +44,7 @@ void FGameplayTagContainerCustomization::CustomizeHeader(TSharedRef<IPropertyHan
 
 		if (CategoriesString.ParseIntoArray(CategoriesStringList, TEXT(","), true) > 0)
 		{
-			UGameplayTagsManager::Get().RequestGameplayTagContainer(CategoriesStringList, MetaDataCategories);
+			UGameplayTagsManager::Get().RequestGameplayTagContainer(CategoriesStringList, MetaDataCategories, false);
 		}
 	}
 	
diff --git a/Engine/Plugins/Editor/GameplayTagsEditor/Source/GameplayTagsEditor/Private/GameplayTagCustomization.cpp b/Engine/Plugins/Editor/GameplayTagsEditor/Source/GameplayTagsEditor/Private/GameplayTagCustomization.cpp
--- a/Engine/Plugins/Editor/GameplayTagsEditor/Source/GameplayTagsEditor/Private/GameplayTagCustomization.cpp
+++ b/Engine/Plugins/Editor/GameplayTagsEditor/Source/GameplayTagsEditor/Private/GameplayTagCustomization.cpp
@@ -51,7 +51,7 @@ void FGameplayTagCustomization::CustomizeHeader(TSharedRef<IPropertyHandle> InSt
 
 		if (CategoriesString.ParseIntoArray(CategoriesStringList, TEXT(","), true) > 0)
 		{
-			UGameplayTagsManager::Get().RequestGameplayTagContainer(CategoriesStringList, MetaDataCategories);
+			UGameplayTagsManager::Get().RequestGameplayTagContainer(CategoriesStringList, MetaDataCategories, false);
 		}
 	}
 
