GameplayTagRoots work for BP structs as well.

diff --git a/Engine/Plugins/Editor/GameplayTagsEditor/Source/GameplayTagsEditor/Private/GameplayTagContainerCustomization.cpp b/Engine/Plugins/Editor/GameplayTagsEditor/Source/GameplayTagsEditor/Private/GameplayTagContainerCustomization.cpp
--- a/Engine/Plugins/Editor/GameplayTagsEditor/Source/GameplayTagsEditor/Private/GameplayTagContainerCustomization.cpp
+++ b/Engine/Plugins/Editor/GameplayTagsEditor/Source/GameplayTagsEditor/Private/GameplayTagContainerCustomization.cpp
@@ -50,7 +50,7 @@ void FGameplayTagContainerCustomization::CustomizeHeader(TSharedRef<IPropertyHan
 	
 	FGameplayTagContainer CurrentTagContainer;
 	CurrentTagContainer.FromExportString(CurrentValue);
-	bool bHasValidCategories = CurrentTagContainer.HasAll(MetaDataCategories);
+	bool bHasValidCategories = CurrentTagContainer.IsEmpty() || MetaDataCategories.IsEmpty() || CurrentTagContainer.HasAll(MetaDataCategories);
 
 	HeaderRow
 		.NameContent()
@@ -129,17 +129,30 @@ void FGameplayTagContainerCustomization::CustomizeChildren(TSharedRef<IPropertyH
 void FGameplayTagContainerCustomization::OnGameplayTagContainerComboClosed(const FGameplayTagContainer& NewTagContainer)
 {
 	MetaDataCategories = NewTagContainer;
-	if (!MetaDataCategories.IsEmpty())
+	FString ContainerString = MetaDataCategories.ToStringSimple();
+	ContainerString.ReplaceInline(TEXT(" "), TEXT(""));
+
+	if (UBlueprint* CreatorBPPtr = CreatorBP.Get())
 	{
-		FString CurrentTag;
-		FString ContainerString = MetaDataCategories.ToStringSimple();
-		ContainerString.ReplaceInline(TEXT(" "), TEXT(""));
-		StructPropertyHandle->GetProperty()->SetMetaData(TEXT("Categories"), *ContainerString);
-		FBlueprintEditorUtils::SetBlueprintVariableMetaData(CreatorBP.Get(), StructPropertyHandle->GetProperty()->GetFName(), StructPropertyHandle->GetProperty()->GetOwner<UFunction>(), TEXT("Categories"), ContainerString);
+		if (!MetaDataCategories.IsEmpty())
+		{
+			FBlueprintEditorUtils::SetBlueprintVariableMetaData(CreatorBPPtr, StructPropertyHandle->GetProperty()->GetFName(), StructPropertyHandle->GetProperty()->GetOwner<UFunction>(), TEXT("Categories"), ContainerString);
+		}
+		else
+		{
+			FBlueprintEditorUtils::RemoveBlueprintVariableMetaData(CreatorBPPtr, StructPropertyHandle->GetProperty()->GetFName(), StructPropertyHandle->GetProperty()->GetOwner<UFunction>(), TEXT("Categories"));
+		}
 	}
 	else
 	{
-		FBlueprintEditorUtils::RemoveBlueprintVariableMetaData(CreatorBP.Get(), StructPropertyHandle->GetProperty()->GetFName(), StructPropertyHandle->GetProperty()->GetOwner<UFunction>(), TEXT("Categories"));
+		if (!MetaDataCategories.IsEmpty())
+		{
+			StructPropertyHandle->GetProperty()->SetMetaData(TEXT("Categories"), *ContainerString);
+		}
+		else
+		{
+			StructPropertyHandle->GetProperty()->RemoveMetaData(TEXT("Categories"));
+		}
 	}
 }
 
diff --git a/Engine/Plugins/Editor/GameplayTagsEditor/Source/GameplayTagsEditor/Private/GameplayTagCustomization.cpp b/Engine/Plugins/Editor/GameplayTagsEditor/Source/GameplayTagsEditor/Private/GameplayTagCustomization.cpp
--- a/Engine/Plugins/Editor/GameplayTagsEditor/Source/GameplayTagsEditor/Private/GameplayTagCustomization.cpp
+++ b/Engine/Plugins/Editor/GameplayTagsEditor/Source/GameplayTagsEditor/Private/GameplayTagCustomization.cpp
@@ -143,19 +143,30 @@ void FGameplayTagCustomization::CustomizeChildren(TSharedRef<IPropertyHandle> In
 void FGameplayTagCustomization::OnGameplayTagContainerComboClosed(const FGameplayTagContainer& NewTagContainer)
 {
 	MetaDataCategories = NewTagContainer;
-	if (!MetaDataCategories.IsEmpty())
+	FString ContainerString = MetaDataCategories.ToStringSimple();
+	ContainerString.ReplaceInline(TEXT(" "), TEXT(""));
+
+	if (UBlueprint* CreatorBPPtr = CreatorBP.Get())
 	{
-		FString CurrentTag;
-		FString ContainerString = MetaDataCategories.ToStringSimple();
-		ContainerString.ReplaceInline(TEXT(" "), TEXT(""));
-		StructPropertyHandle->GetProperty()->SetMetaData(TEXT("Categories"), *ContainerString);
-		FBlueprintEditorUtils::SetBlueprintVariableMetaData(CreatorBP.Get(), StructPropertyHandle->GetProperty()->GetFName(), 
-		StructPropertyHandle->GetProperty()->GetOwner<UFunction>(), TEXT("Categories"), ContainerString);
+		if (!MetaDataCategories.IsEmpty())
+		{
+			FBlueprintEditorUtils::SetBlueprintVariableMetaData(CreatorBPPtr, StructPropertyHandle->GetProperty()->GetFName(), StructPropertyHandle->GetProperty()->GetOwner<UFunction>(), TEXT("Categories"), ContainerString);
+		}
+		else
+		{
+			FBlueprintEditorUtils::RemoveBlueprintVariableMetaData(CreatorBPPtr, StructPropertyHandle->GetProperty()->GetFName(), StructPropertyHandle->GetProperty()->GetOwner<UFunction>(), TEXT("Categories"));
+		}
 	}
 	else
 	{
-		FBlueprintEditorUtils::RemoveBlueprintVariableMetaData(CreatorBP.Get(), StructPropertyHandle->GetProperty()->GetFName(), 
-		StructPropertyHandle->GetProperty()->GetOwner<UFunction>(), TEXT("Categories"));
+		if (!MetaDataCategories.IsEmpty())
+		{
+			StructPropertyHandle->GetProperty()->SetMetaData(TEXT("Categories"), *ContainerString);
+		}
+		else
+		{
+			StructPropertyHandle->GetProperty()->RemoveMetaData(TEXT("Categories"));
+		}
 	}
 }
 
