Gameplay tags are inlined as map keys.
#11708

diff --git a/Engine/Plugins/Editor/GameplayTagsEditor/Source/GameplayTagsEditor/Private/GameplayTagCustomization.h b/Engine/Plugins/Editor/GameplayTagsEditor/Source/GameplayTagsEditor/Private/GameplayTagCustomization.h
--- a/Engine/Plugins/Editor/GameplayTagsEditor/Source/GameplayTagsEditor/Private/GameplayTagCustomization.h
+++ b/Engine/Plugins/Editor/GameplayTagsEditor/Source/GameplayTagsEditor/Private/GameplayTagCustomization.h
@@ -22,8 +22,11 @@ public:
 	/** Overridden to show an edit button to launch the gameplay tag editor */
 	virtual void CustomizeHeader(TSharedRef<IPropertyHandle> StructPropertyHandle, FDetailWidgetRow& HeaderRow, IPropertyTypeCustomizationUtils& StructCustomizationUtils) override;
 
-	/** Overridden to do nothing */
+	/** Overridden to display "Categories" metadata if created from BP */
 	virtual void CustomizeChildren(TSharedRef<IPropertyHandle> InStructPropertyHandle, IDetailChildrenBuilder& ChildBuilder, IPropertyTypeCustomizationUtils& StructCustomizationUtils) override;
+	
+	/** Overridden to show tags and value inline */
+	virtual bool ShouldInlineKey() const override { return true; }
 
 private:
 	
