Updated headers using UnrealCodeFixup to make sure dllstorage is on methods/staticvars instead of types. Part 2/n

diff --git a/Engine/Plugins/Editor/GameplayTagsEditor/Source/GameplayTagsEditor/Public/AssetDefinition_GameplayTagAssetBase.h b/Engine/Plugins/Editor/GameplayTagsEditor/Source/GameplayTagsEditor/Public/AssetDefinition_GameplayTagAssetBase.h
--- a/Engine/Plugins/Editor/GameplayTagsEditor/Source/GameplayTagsEditor/Public/AssetDefinition_GameplayTagAssetBase.h
+++ b/Engine/Plugins/Editor/GameplayTagsEditor/Source/GameplayTagsEditor/Public/AssetDefinition_GameplayTagAssetBase.h
@@ -6,12 +6,14 @@
 
 #include "AssetDefinition_GameplayTagAssetBase.generated.h"
 
+#define UE_API GAMEPLAYTAGSEDITOR_API
+
 struct FToolMenuSection;
 struct FGameplayTagContainer;
 
 /** Base asset type actions for any classes with gameplay tagging */
-UCLASS(Abstract)
-class GAMEPLAYTAGSEDITOR_API UAssetDefinition_GameplayTagAssetBase : public UAssetDefinitionDefault
+UCLASS(MinimalAPI, Abstract)
+class UAssetDefinition_GameplayTagAssetBase : public UAssetDefinitionDefault
 {
 	GENERATED_BODY()
 
@@ -20,7 +22,7 @@ public:
 
 	/** Traditionally these are implemented in a MenuExtension namespace. However, UAssetDefinition_GameplayTagAssetBase is an abstract class,
 	* and the derived classes need to invoke this in their static MenuExtension functions**/
-	static void AddGameplayTagsEditMenuExtension(FToolMenuSection& InSection, TArray<UObject*> InObjects, const FName& OwnedGameplayTagPropertyName);
+	static UE_API void AddGameplayTagsEditMenuExtension(FToolMenuSection& InSection, TArray<UObject*> InObjects, const FName& OwnedGameplayTagPropertyName);
 
 	// UAssetDefinition Begin
 	/** Overridden to specify misc category */
@@ -37,5 +39,7 @@ private:
 	 *
 	 * @param TagAssets	Assets to open the editor with
 	 */
-	static void OpenGameplayTagEditor(TArray<UObject*> Objects, TArray<FGameplayTagContainer> Containers, const FName& OwnedGameplayTagPropertyName);
+	static UE_API void OpenGameplayTagEditor(TArray<UObject*> Objects, TArray<FGameplayTagContainer> Containers, const FName& OwnedGameplayTagPropertyName);
 };
+
+#undef UE_API
diff --git a/Engine/Plugins/Editor/GameplayTagsEditor/Source/GameplayTagsEditor/Public/GameplayTagsEditorModule.h b/Engine/Plugins/Editor/GameplayTagsEditor/Source/GameplayTagsEditor/Public/GameplayTagsEditorModule.h
--- a/Engine/Plugins/Editor/GameplayTagsEditor/Source/GameplayTagsEditor/Public/GameplayTagsEditorModule.h
+++ b/Engine/Plugins/Editor/GameplayTagsEditor/Source/GameplayTagsEditor/Public/GameplayTagsEditorModule.h
@@ -4,6 +4,8 @@
 
 #include "Modules/ModuleManager.h"
 
+#define UE_API GAMEPLAYTAGSEDITOR_API
+
 class IPropertyTypeCustomization;
 struct FGameplayTag;
 struct FGameplayTagContainer;
@@ -45,26 +47,26 @@ public:
 	}
 
 	/** Tries to add a new gameplay tag to the ini lists */
-	GAMEPLAYTAGSEDITOR_API virtual bool AddNewGameplayTagToINI(const FString& NewTag, const FString& Comment = TEXT(""), FName TagSourceName = NAME_None, bool bIsRestrictedTag = false, bool bAllowNonRestrictedChildren = true) = 0;
+	virtual bool AddNewGameplayTagToINI(const FString& NewTag, const FString& Comment = TEXT(""), FName TagSourceName = NAME_None, bool bIsRestrictedTag = false, bool bAllowNonRestrictedChildren = true) = 0;
 
 	/** Tries to delete a tag from the library. This will pop up special UI or error messages as needed. It will also delete redirectors if that is specified. */
-	GAMEPLAYTAGSEDITOR_API virtual bool DeleteTagFromINI(TSharedPtr<struct FGameplayTagNode> TagNodeToDelete) = 0;
-	GAMEPLAYTAGSEDITOR_API virtual void DeleteTagsFromINI(const TArray<TSharedPtr<struct FGameplayTagNode>>& TagNodesToDelete) = 0;
+	virtual bool DeleteTagFromINI(TSharedPtr<struct FGameplayTagNode> TagNodeToDelete) = 0;
+	virtual void DeleteTagsFromINI(const TArray<TSharedPtr<struct FGameplayTagNode>>& TagNodesToDelete) = 0;
 
 	/** Tries to rename a tag, leaving a rediretor in the ini, and adding the new tag if it does not exist yet */
-	GAMEPLAYTAGSEDITOR_API virtual bool RenameTagInINI(const FString& TagToRename, const FString& TagToRenameTo) = 0;
+	virtual bool RenameTagInINI(const FString& TagToRename, const FString& TagToRenameTo) = 0;
 
 	/** Updates info about a tag */
-	GAMEPLAYTAGSEDITOR_API virtual bool UpdateTagInINI(const FString& TagToUpdate, const FString& Comment, bool bIsRestrictedTag, bool bAllowNonRestrictedChildren) = 0;
+	virtual bool UpdateTagInINI(const FString& TagToUpdate, const FString& Comment, bool bIsRestrictedTag, bool bAllowNonRestrictedChildren) = 0;
 
 	/** Tries to move existing tags from each source ini lists to the target source ini list. */
-	GAMEPLAYTAGSEDITOR_API virtual bool MoveTagsBetweenINI(const TArray<FString>& TagsToMove, const FName& TargetTagSource, TArray<FString>& OutTagsMoved, TArray<FString>& OutFailedToMoveTags) = 0;
+	virtual bool MoveTagsBetweenINI(const TArray<FString>& TagsToMove, const FName& TargetTagSource, TArray<FString>& OutTagsMoved, TArray<FString>& OutFailedToMoveTags) = 0;
 
 	/** Adds a transient gameplay tag (only valid for the current editor session) */
-	GAMEPLAYTAGSEDITOR_API virtual bool AddTransientEditorGameplayTag(const FString& NewTransientTag) = 0;
+	virtual bool AddTransientEditorGameplayTag(const FString& NewTransientTag) = 0;
 
 	/** Adds a new tag source, well use project config directory if not specified. This will not save anything until a tag is added */
-	GAMEPLAYTAGSEDITOR_API virtual bool AddNewGameplayTagSource(const FString& NewTagSource, const FString& RootDirToUse = FString()) = 0;
+	virtual bool AddNewGameplayTagSource(const FString& NewTagSource, const FString& RootDirToUse = FString()) = 0;
 
 	/**
 	 * Creates a simple version of a tag container widget that has a default value and will call a custom callback
@@ -72,7 +74,7 @@ public:
 	 * @param GameplayTagValue	Shared ptr to tag container value that will be used as default and modified
 	 * @param FilterString		Optional filter string, same format as Categories metadata on tag properties
 	 */
-	GAMEPLAYTAGSEDITOR_API virtual TSharedRef<SWidget> MakeGameplayTagContainerWidget(FOnSetGameplayTagContainer OnSetTag, TSharedPtr<FGameplayTagContainer> GameplayTagContainer, const FString& FilterString = FString()) = 0;
+	virtual TSharedRef<SWidget> MakeGameplayTagContainerWidget(FOnSetGameplayTagContainer OnSetTag, TSharedPtr<FGameplayTagContainer> GameplayTagContainer, const FString& FilterString = FString()) = 0;
 
 	/**
 	 * Creates a simple version of a gameplay tag widget that has a default value and will call a custom callback
@@ -80,29 +82,31 @@ public:
 	 * @param GameplayTagValue	Shared ptr to tag value that will be used as default and modified
 	 * @param FilterString		Optional filter string, same format as Categories metadata on tag properties
 	 */
-	GAMEPLAYTAGSEDITOR_API virtual TSharedRef<SWidget> MakeGameplayTagWidget(FOnSetGameplayTag OnSetTag, TSharedPtr<FGameplayTag> GameplayTag, const FString& FilterString = FString()) = 0;
+	virtual TSharedRef<SWidget> MakeGameplayTagWidget(FOnSetGameplayTag OnSetTag, TSharedPtr<FGameplayTag> GameplayTag, const FString& FilterString = FString()) = 0;
 
 	/** Returns the list of gameplay tags that are not used by content */
-	GAMEPLAYTAGSEDITOR_API virtual void GetUnusedGameplayTags(TArray<TSharedPtr<struct FGameplayTagNode>>& OutUnusedTags) = 0;
+	virtual void GetUnusedGameplayTags(TArray<TSharedPtr<struct FGameplayTagNode>>& OutUnusedTags) = 0;
 };
 
 /** This is public so that child structs of FGameplayTag can use the details customization */
-struct GAMEPLAYTAGSEDITOR_API FGameplayTagCustomizationPublic
+struct FGameplayTagCustomizationPublic
 {
-	static TSharedRef<IPropertyTypeCustomization> MakeInstance();
+	static UE_API TSharedRef<IPropertyTypeCustomization> MakeInstance();
 	UE_DEPRECATED(5.3, "FGameplayTagCustomizationOptions are deprecated. The options are not used anymore, use the version without options.")
-	static TSharedRef<IPropertyTypeCustomization> MakeInstanceWithOptions(const FGameplayTagCustomizationOptions& Options);
+	static UE_API TSharedRef<IPropertyTypeCustomization> MakeInstanceWithOptions(const FGameplayTagCustomizationOptions& Options);
 };
 
 /** This is public so that child structs of FGameplayTagContainer can use the details customization */
-struct GAMEPLAYTAGSEDITOR_API FGameplayTagContainerCustomizationPublic
+struct FGameplayTagContainerCustomizationPublic
 {
-	static TSharedRef<IPropertyTypeCustomization> MakeInstance();
+	static UE_API TSharedRef<IPropertyTypeCustomization> MakeInstance();
 	UE_DEPRECATED(5.3, "FGameplayTagContainerCustomizationOptions are deprecated. The options are not used anymore, use the version without options.")
-	static TSharedRef<IPropertyTypeCustomization> MakeInstanceWithOptions(const FGameplayTagContainerCustomizationOptions& Options);
+	static UE_API TSharedRef<IPropertyTypeCustomization> MakeInstanceWithOptions(const FGameplayTagContainerCustomizationOptions& Options);
 };
 
-struct GAMEPLAYTAGSEDITOR_API FRestrictedGameplayTagCustomizationPublic
+struct FRestrictedGameplayTagCustomizationPublic
 {
-	static TSharedRef<IPropertyTypeCustomization> MakeInstance();
+	static UE_API TSharedRef<IPropertyTypeCustomization> MakeInstance();
 };
+
+#undef UE_API
diff --git a/Engine/Plugins/Editor/GameplayTagsEditor/Source/GameplayTagsEditor/Public/SGameplayTagPicker.h b/Engine/Plugins/Editor/GameplayTagsEditor/Source/GameplayTagsEditor/Public/SGameplayTagPicker.h
--- a/Engine/Plugins/Editor/GameplayTagsEditor/Source/GameplayTagsEditor/Public/SGameplayTagPicker.h
+++ b/Engine/Plugins/Editor/GameplayTagsEditor/Source/GameplayTagsEditor/Public/SGameplayTagPicker.h
@@ -14,6 +14,8 @@
 #include "Widgets/Views/STreeView.h"
 #include "GameplayTagsManager.h"
 
+#define UE_API GAMEPLAYTAGSEDITOR_API
+
 class IPropertyHandle;
 class SComboButton;
 class SAddNewGameplayTagWidget;
@@ -27,7 +29,7 @@ enum class EGameplayTagPickerMode : uint8
 };
 
 /** Widget allowing user to tag assets with gameplay tags */
-class GAMEPLAYTAGSEDITOR_API SGameplayTagPicker : public SCompoundWidget
+class SGameplayTagPicker : public SCompoundWidget
 {
 public:
 
@@ -105,44 +107,44 @@ public:
 	 * Given a property handle, try and enumerate the tag containers from within it (when dealing with a struct property of type FGameplayTag or FGameplayTagContainer).
 	 * @return True if it was possible to enumerate containers (even if no containers were enumerated), or false otherwise.
 	 */
-	static bool EnumerateEditableTagContainersFromPropertyHandle(const TSharedRef<IPropertyHandle>& PropHandle, TFunctionRef<bool(const FGameplayTagContainer&)> Callback);
+	static UE_API bool EnumerateEditableTagContainersFromPropertyHandle(const TSharedRef<IPropertyHandle>& PropHandle, TFunctionRef<bool(const FGameplayTagContainer&)> Callback);
 
 	/**
 	 * Given a property handle, try and extract the editable tag containers from within it (when dealing with a struct property of type FGameplayTagContainer).
 	 * @return True if it was possible to extract containers (even if no containers were extracted), or false otherwise.
 	 */
-	static bool GetEditableTagContainersFromPropertyHandle(const TSharedRef<IPropertyHandle>& PropHandle, TArray<FGameplayTagContainer>& OutEditableContainers);
+	static UE_API bool GetEditableTagContainersFromPropertyHandle(const TSharedRef<IPropertyHandle>& PropHandle, TArray<FGameplayTagContainer>& OutEditableContainers);
 
-	virtual ~SGameplayTagPicker() override;
+	UE_API virtual ~SGameplayTagPicker() override;
 
 	/** Construct the actual widget */
-	void Construct(const FArguments& InArgs);
+	UE_API void Construct(const FArguments& InArgs);
 	
-	virtual void Tick(const FGeometry& AllottedGeometry, const double InCurrentTime, const float InDeltaTime) override;
+	UE_API virtual void Tick(const FGeometry& AllottedGeometry, const double InCurrentTime, const float InDeltaTime) override;
 
 	/** Updates the tag list when the filter text changes */
-	void OnFilterTextChanged(const FText& InFilterText);
+	UE_API void OnFilterTextChanged(const FText& InFilterText);
 
 	/** Returns true if this TagNode has any children that match the current filter */
-	bool FilterChildrenCheck(TSharedPtr<FGameplayTagNode>& InItem) const;
+	UE_API bool FilterChildrenCheck(TSharedPtr<FGameplayTagNode>& InItem) const;
 
 	/** Child recursion helper for FilterChildrenCheck() */
-	bool FilterChildrenCheckRecursive(TSharedPtr<FGameplayTagNode>& InItem) const;
+	UE_API bool FilterChildrenCheckRecursive(TSharedPtr<FGameplayTagNode>& InItem) const;
 
 	/** Refreshes the tags that should be displayed by the widget */
-	void RefreshTags();
+	UE_API void RefreshTags();
 
 	/** Forces the widget to refresh its tags on the next tick */
-	void RefreshOnNextTick();
+	UE_API void RefreshOnNextTick();
 
 	/** Scrolls the view to specified tag. */
-	void RequestScrollToView(const FGameplayTag Tag);
+	UE_API void RequestScrollToView(const FGameplayTag Tag);
 	
 	/** Gets the widget to focus once the menu opens. */
-	TSharedPtr<SWidget> GetWidgetToFocusOnOpen();
+	UE_API TSharedPtr<SWidget> GetWidgetToFocusOnOpen();
 
 	/** Sets tag containers to edit. */
-	void SetTagContainers(TConstArrayView<FGameplayTagContainer> TagContainers);
+	UE_API void SetTagContainers(TConstArrayView<FGameplayTagContainer> TagContainers);
 	
 private:
 
@@ -153,19 +155,19 @@ private:
 		Duplicate,
 	};
 
-	void OnPostUndoRedo();
+	UE_API void OnPostUndoRedo();
 
 	/** Verify the tags are all valid and if not prompt the user. */
-	void VerifyAssetTagValidity();
+	UE_API void VerifyAssetTagValidity();
 
 	/* Filters the tree view based on the current filter text. */
-	void FilterTagTree();
+	UE_API void FilterTagTree();
 
 	/* string that sets the section of the ini file to use for this class*/ 
-	static const FString SettingsIniSection;
+	static UE_API const FString SettingsIniSection;
 
 	/* Filename of ini file to store state of the UI */
-	static const FString& GetGameplayTagsEditorStateIni();
+	static UE_API const FString& GetGameplayTagsEditorStateIni();
 
 	/**
 	 * Generate a row widget for the specified item node and table
@@ -175,7 +177,7 @@ private:
 	 * 
 	 * @return Generated row widget for the item node
 	 */
-	TSharedRef<ITableRow> OnGenerateRow(TSharedPtr<FGameplayTagNode> InItem, const TSharedRef<STableViewBase>& OwnerTable);
+	UE_API TSharedRef<ITableRow> OnGenerateRow(TSharedPtr<FGameplayTagNode> InItem, const TSharedRef<STableViewBase>& OwnerTable);
 
 	/**
 	 * Get children nodes of the specified node
@@ -183,7 +185,7 @@ private:
 	 * @param InItem		Node to get children of
 	 * @param OutChildren	[OUT] Array of children nodes, if any
 	 */
-	void OnGetChildren(TSharedPtr<FGameplayTagNode> InItem, TArray< TSharedPtr<FGameplayTagNode> >& OutChildren);
+	UE_API void OnGetChildren(TSharedPtr<FGameplayTagNode> InItem, TArray< TSharedPtr<FGameplayTagNode> >& OutChildren);
 
 	/**
 	 * Called via delegate when the status of a check box in a row changes
@@ -191,7 +193,7 @@ private:
 	 * @param NewCheckState	New check box state
 	 * @param NodeChanged	Node that was checked/unchecked
 	 */
-	void OnTagCheckStatusChanged(ECheckBoxState NewCheckState, TSharedPtr<FGameplayTagNode> NodeChanged);
+	UE_API void OnTagCheckStatusChanged(ECheckBoxState NewCheckState, TSharedPtr<FGameplayTagNode> NodeChanged);
 
 	/**
 	 * Called via delegate to determine the checkbox state of the specified node
@@ -200,12 +202,12 @@ private:
 	 * 
 	 * @return Checkbox state of the specified node
 	 */
-	ECheckBoxState IsTagChecked(TSharedPtr<FGameplayTagNode> Node) const;
+	UE_API ECheckBoxState IsTagChecked(TSharedPtr<FGameplayTagNode> Node) const;
 
 	/**
 	 * @return true if the exact Tag provided is included in any of the tag containers the widget is editing.
 	 */
-	bool IsExactTagInCollection(TSharedPtr<FGameplayTagNode> Node) const;
+	UE_API bool IsExactTagInCollection(TSharedPtr<FGameplayTagNode> Node) const;
 
 	/**
 	 * Called via delegate when the status of the allow non-restricted children check box in a row changes
@@ -213,7 +215,7 @@ private:
 	 * @param NewCheckState	New check box state
 	 * @param NodeChanged	Node that was checked/unchecked
 	 */
-	void OnAllowChildrenTagCheckStatusChanged(ECheckBoxState NewCheckState, TSharedPtr<FGameplayTagNode> NodeChanged);
+	UE_API void OnAllowChildrenTagCheckStatusChanged(ECheckBoxState NewCheckState, TSharedPtr<FGameplayTagNode> NodeChanged);
 
 	/**
 	 * Called via delegate to determine the non-restricted children checkbox state of the specified node
@@ -222,24 +224,24 @@ private:
 	 *
 	 * @return Non-restricted children heckbox state of the specified node
 	 */
-	ECheckBoxState IsAllowChildrenTagChecked(TSharedPtr<FGameplayTagNode> Node) const;
+	UE_API ECheckBoxState IsAllowChildrenTagChecked(TSharedPtr<FGameplayTagNode> Node) const;
 
 	/** Helper function to determine the visibility of the checkbox for allowing non-restricted children of restricted gameplay tags */
-	EVisibility DetermineAllowChildrenVisible(TSharedPtr<FGameplayTagNode> Node) const;
+	UE_API EVisibility DetermineAllowChildrenVisible(TSharedPtr<FGameplayTagNode> Node) const;
 
 	/**
 	 * Helper function called when the specified node is checked
 	 * 
 	 * @param NodeChecked	Node that was checked by the user
 	 */
-	void OnTagChecked(TSharedPtr<FGameplayTagNode> NodeChecked);
+	UE_API void OnTagChecked(TSharedPtr<FGameplayTagNode> NodeChecked);
 
 	/**
 	 * Helper function called when the specified node is unchecked
 	 * 
 	 * @param NodeUnchecked	Node that was unchecked by the user
 	 */
-	void OnTagUnchecked(TSharedPtr<FGameplayTagNode> NodeUnchecked);
+	UE_API void OnTagUnchecked(TSharedPtr<FGameplayTagNode> NodeUnchecked);
 
 	/**
 	 * Recursive function to uncheck all child tags
@@ -247,7 +249,7 @@ private:
 	 * @param NodeUnchecked	Node that was unchecked by the user
 	 * @param EditableContainer The container we are removing the tags from
 	 */
-	void UncheckChildren(TSharedPtr<FGameplayTagNode> NodeUnchecked, FGameplayTagContainer& EditableContainer);
+	UE_API void UncheckChildren(TSharedPtr<FGameplayTagNode> NodeUnchecked, FGameplayTagContainer& EditableContainer);
 
 	/**
 	 * Called via delegate to determine the text colour of the specified node
@@ -256,19 +258,19 @@ private:
 	 *
 	 * @return Text colour of the specified node
 	 */
-	FSlateColor GetTagTextColour(TSharedPtr<FGameplayTagNode> Node) const;
+	UE_API FSlateColor GetTagTextColour(TSharedPtr<FGameplayTagNode> Node) const;
 
 	/** Called when the user clicks the "Create and Manage Gameplay Tags" button; Opens Gameplay Tag manager window. */
-	void OnManageTagsClicked(TSharedPtr<FGameplayTagNode> Node, TSharedPtr<SComboButton> OwnerCombo);
+	UE_API void OnManageTagsClicked(TSharedPtr<FGameplayTagNode> Node, TSharedPtr<SComboButton> OwnerCombo);
 
 	/** Called when the user clicks the "Clear All" button; Clears all tags */
-	void OnClearAllClicked(TSharedPtr<SComboButton> OwnerCombo);
+	UE_API void OnClearAllClicked(TSharedPtr<SComboButton> OwnerCombo);
 
 	/** Called when the user clicks the "Expand All" button; Expands the entire tag tree */
-	void OnExpandAllClicked(TSharedPtr<SComboButton> OwnerCombo);
+	UE_API void OnExpandAllClicked(TSharedPtr<SComboButton> OwnerCombo);
 
 	/** Called when the user clicks the "Collapse All" button; Collapses the entire tag tree */
-	void OnCollapseAllClicked(TSharedPtr<SComboButton> OwnerCombo);
+	UE_API void OnCollapseAllClicked(TSharedPtr<SComboButton> OwnerCombo);
 
 	/**
 	 * Helper function to set the expansion state of the tree widget
@@ -276,7 +278,7 @@ private:
 	 * @param bExpand If true, expand the entire tree; Otherwise, collapse the entire tree
 	 * @param bPersistExpansion If true, persist the expansion state.
 	 */
-	void SetTagTreeItemExpansion(bool bExpand, bool bPersistExpansion = false);
+	UE_API void SetTagTreeItemExpansion(bool bExpand, bool bPersistExpansion = false);
 
 	/**
 	 * Helper function to set the expansion state of a specific node
@@ -285,100 +287,100 @@ private:
 	 * @param bExpand	If true, expand the node; Otherwise, collapse the node
 	 * @param bPersistExpansion If true, persist the expansion state.
 	 */
-	void SetTagNodeItemExpansion(TSharedPtr<FGameplayTagNode> Node, bool bExpand, bool bPersistExpansion = false);
+	UE_API void SetTagNodeItemExpansion(TSharedPtr<FGameplayTagNode> Node, bool bExpand, bool bPersistExpansion = false);
 
 	/** Load settings for the tags*/
-	void LoadSettings();
+	UE_API void LoadSettings();
 
 	/** Migrate settings */
-	void MigrateSettings();
+	UE_API void MigrateSettings();
 
 	/** Helper function to determine the visibility of the expandable UI controls */
-	EVisibility DetermineExpandableUIVisibility() const;
+	UE_API EVisibility DetermineExpandableUIVisibility() const;
 
 	/** Helper function to determine the visibility of the Add New Tag widget */
-	bool CanAddNewTag() const;
+	UE_API bool CanAddNewTag() const;
 
 	/** Helper function to determine the visibility of the Add New Subtag widget */
-	bool CanAddNewSubTag(TSharedPtr<FGameplayTagNode> Node) const;
+	UE_API bool CanAddNewSubTag(TSharedPtr<FGameplayTagNode> Node) const;
 
 	/** Helper function to determine if tag can be modified */
-	bool CanModifyTag(TSharedPtr<FGameplayTagNode> Node) const;
+	UE_API bool CanModifyTag(TSharedPtr<FGameplayTagNode> Node) const;
 
 	/** Recursive load function to go through all tags in the tree and set the expansion*/
-	void LoadTagNodeItemExpansion(TSharedPtr<FGameplayTagNode> Node );
+	UE_API void LoadTagNodeItemExpansion(TSharedPtr<FGameplayTagNode> Node );
 
 	/** Recursive function to go through all tags in the tree and set the expansion to default*/
-	void SetDefaultTagNodeItemExpansion(TSharedPtr<FGameplayTagNode> Node);
+	UE_API void SetDefaultTagNodeItemExpansion(TSharedPtr<FGameplayTagNode> Node);
 
 	/** Expansion changed callback */
-	void OnExpansionChanged(TSharedPtr<FGameplayTagNode> InItem, bool bIsExpanded);
+	UE_API void OnExpansionChanged(TSharedPtr<FGameplayTagNode> InItem, bool bIsExpanded);
 
 	/** Callback for when a new tag is added */
-	void OnGameplayTagAdded(const FString& TagName, const FString& TagComment, const FName& TagSource);
+	UE_API void OnGameplayTagAdded(const FString& TagName, const FString& TagComment, const FName& TagSource);
 
 	/** Initializes the inline add tag widget and makes in visible. */
-	void ShowInlineAddTagWidget(const EGameplayTagAdd Mode, TSharedPtr<FGameplayTagNode> InTagNode = TSharedPtr<FGameplayTagNode>());
+	UE_API void ShowInlineAddTagWidget(const EGameplayTagAdd Mode, TSharedPtr<FGameplayTagNode> InTagNode = TSharedPtr<FGameplayTagNode>());
 
 	/** Returns true if the passed in node is the root node, and there is only one node on the root level */
-	bool IsNodeSingleRoot(const TSharedPtr<FGameplayTagNode>& InNode) const;
+	UE_API bool IsNodeSingleRoot(const TSharedPtr<FGameplayTagNode>& InNode) const;
 
-	FText GetHighlightText() const;
+	UE_API FText GetHighlightText() const;
 	
 	/** Creates a dropdown menu to provide additional functionality for tags (renaming, deletion, search for references, etc.) */
-	TSharedRef<SWidget> MakeTagActionsMenu(TSharedPtr<FGameplayTagNode> InTagNode, TSharedPtr<SComboButton> OwnerCombo, const bool bInShouldCloseWindowAfterMenuSelection);
+	UE_API TSharedRef<SWidget> MakeTagActionsMenu(TSharedPtr<FGameplayTagNode> InTagNode, TSharedPtr<SComboButton> OwnerCombo, const bool bInShouldCloseWindowAfterMenuSelection);
 
 	/** Creates settings menu content. */
-	TSharedRef<SWidget> MakeSettingsMenu(TSharedPtr<SComboButton> OwnerCombo);
+	UE_API TSharedRef<SWidget> MakeSettingsMenu(TSharedPtr<SComboButton> OwnerCombo);
 
-	void OnAddSubTag(TSharedPtr<FGameplayTagNode> InTagNode, TSharedPtr<SComboButton> OwnerCombo);
+	UE_API void OnAddSubTag(TSharedPtr<FGameplayTagNode> InTagNode, TSharedPtr<SComboButton> OwnerCombo);
 
-	void OnDuplicateTag(TSharedPtr<FGameplayTagNode> InTagNode, TSharedPtr<SComboButton> OwnerCombo);
+	UE_API void OnDuplicateTag(TSharedPtr<FGameplayTagNode> InTagNode, TSharedPtr<SComboButton> OwnerCombo);
 	
 	/** Attempts to rename the tag through a dialog box */
-	void OnRenameTag(TSharedPtr<FGameplayTagNode> InTagNode, TSharedPtr<SComboButton> OwnerCombo);
+	UE_API void OnRenameTag(TSharedPtr<FGameplayTagNode> InTagNode, TSharedPtr<SComboButton> OwnerCombo);
 
 	/** Attempts to delete the specified tag */
-	void OnDeleteTag(TSharedPtr<FGameplayTagNode> InTagNode, TSharedPtr<SComboButton> OwnerCombo);
+	UE_API void OnDeleteTag(TSharedPtr<FGameplayTagNode> InTagNode, TSharedPtr<SComboButton> OwnerCombo);
 
 	/** Attempts to select the exact specified tag*/
-	void OnSelectExactTag(TSharedPtr<FGameplayTagNode> InTagNode, TSharedPtr<SComboButton> OwnerCombo);
+	UE_API void OnSelectExactTag(TSharedPtr<FGameplayTagNode> InTagNode, TSharedPtr<SComboButton> OwnerCombo);
 
 	/** Attempts to unselect the specified tag, but not the children */
-	void OnUnselectExactTag(TSharedPtr<FGameplayTagNode> InTagNode, TSharedPtr<SComboButton> OwnerCombo);
+	UE_API void OnUnselectExactTag(TSharedPtr<FGameplayTagNode> InTagNode, TSharedPtr<SComboButton> OwnerCombo);
 
 	/** Searches for all references for the selected tag */
-	void OnSearchForReferences(TSharedPtr<FGameplayTagNode> InTagNode, TSharedPtr<SComboButton> OwnerCombo);
+	UE_API void OnSearchForReferences(TSharedPtr<FGameplayTagNode> InTagNode, TSharedPtr<SComboButton> OwnerCombo);
 
 	/** Searches for all references for the selected tag and its children */
-	void OnSearchForReferencesAsChildren(TSharedPtr<FGameplayTagNode> InTagNode, TSharedPtr<SComboButton> OwnerCombo);
+	UE_API void OnSearchForReferencesAsChildren(TSharedPtr<FGameplayTagNode> InTagNode, TSharedPtr<SComboButton> OwnerCombo);
 
 	/** Copies individual tag's name to clipboard. */
-	void OnCopyTagNameToClipboard(TSharedPtr<FGameplayTagNode> InTagNode, TSharedPtr<SComboButton> OwnerCombo);
+	UE_API void OnCopyTagNameToClipboard(TSharedPtr<FGameplayTagNode> InTagNode, TSharedPtr<SComboButton> OwnerCombo);
 	
 	/** Returns true if the user can select tags from the widget */
-	bool CanSelectTags() const;
+	UE_API bool CanSelectTags() const;
 
 	/** Called when tags are selected. */
-	void OnContainersChanged();
+	UE_API void OnContainersChanged();
 
 	/** Opens a dialog window to rename the selected tag */
-	void OpenRenameGameplayTagDialog(TSharedPtr<FGameplayTagNode> GameplayTagNode) const;
+	UE_API void OpenRenameGameplayTagDialog(TSharedPtr<FGameplayTagNode> GameplayTagNode) const;
 
 	/** Delegate that is fired when a tag is successfully renamed */
-	void OnGameplayTagRenamed(FString OldTagName, FString NewTagName);
+	UE_API void OnGameplayTagRenamed(FString OldTagName, FString NewTagName);
 
 	/** Populate tag items from the gameplay tags manager. */
-	void GetFilteredGameplayRootTags(const FString& InFilterString, TArray<TSharedPtr<FGameplayTagNode>>& OutNodes) const;
+	UE_API void GetFilteredGameplayRootTags(const FString& InFilterString, TArray<TSharedPtr<FGameplayTagNode>>& OutNodes) const;
 
 	/** Called to create context menu for the tag tree items. */
-	TSharedPtr<SWidget> OnTreeContextMenuOpening();
+	UE_API TSharedPtr<SWidget> OnTreeContextMenuOpening();
 
 	/** Called when the tag tree selection changes. */
-	void OnTreeSelectionChanged(TSharedPtr<FGameplayTagNode> SelectedItem, ESelectInfo::Type SelectInfo);
+	UE_API void OnTreeSelectionChanged(TSharedPtr<FGameplayTagNode> SelectedItem, ESelectInfo::Type SelectInfo);
 
 	/** Called to handle key presses in the tag tree. */
-	FReply OnTreeKeyDown(const FGeometry& Geometry, const FKeyEvent& Key);
+	UE_API FReply OnTreeKeyDown(const FGeometry& Geometry, const FKeyEvent& Key);
 	
 	/* Holds the Name of this TagContainer used for saving out expansion settings */
 	FString SettingsName;
@@ -468,4 +470,6 @@ namespace UE::GameplayTags::Editor
 {
 	GAMEPLAYTAGSEDITOR_API TWeakPtr<SGameplayTagPicker> OpenGameplayTagManager(const FGameplayTagManagerWindowArgs& Args);
 	GAMEPLAYTAGSEDITOR_API TSharedRef<SGameplayTagPicker> Create(const FGameplayTagManagerWindowArgs& Args);
-};
\ No newline at end of file
+};
+
+#undef UE_API
diff --git a/Engine/Plugins/Editor/GameplayTagsEditor/Source/GameplayTagsEditor/Public/SGameplayTagWidget.h b/Engine/Plugins/Editor/GameplayTagsEditor/Source/GameplayTagsEditor/Public/SGameplayTagWidget.h
--- a/Engine/Plugins/Editor/GameplayTagsEditor/Source/GameplayTagsEditor/Public/SGameplayTagWidget.h
+++ b/Engine/Plugins/Editor/GameplayTagsEditor/Source/GameplayTagsEditor/Public/SGameplayTagWidget.h
@@ -5,6 +5,8 @@
 #include "Styling/AppStyle.h"
 #include "Widgets/SCompoundWidget.h"
 
+#define UE_API GAMEPLAYTAGSEDITOR_API
+
 class ITableRow;
 class SSearchBox;
 class STableViewBase;
@@ -26,7 +28,7 @@ enum class EGameplayTagUIMode : uint8
 };
 
 /** Widget allowing user to tag assets with gameplay tags */
-class GAMEPLAYTAGSEDITOR_API SGameplayTagWidget : public SCompoundWidget
+class SGameplayTagWidget : public SCompoundWidget
 {
 public:
 
@@ -99,53 +101,53 @@ public:
 	 * Given a property handle, try and enumerate the editable tag containers from within it (when dealing with a struct property of type FGameplayTagContainer).
 	 * @return True if it was possible to enumerate containers (even if no containers were enumerated), or false otherwise.
 	 */
-	static bool EnumerateEditableTagContainersFromPropertyHandle(const TSharedRef<IPropertyHandle>& PropHandle, TFunctionRef<bool(const FEditableGameplayTagContainerDatum&)> Callback);
+	static UE_API bool EnumerateEditableTagContainersFromPropertyHandle(const TSharedRef<IPropertyHandle>& PropHandle, TFunctionRef<bool(const FEditableGameplayTagContainerDatum&)> Callback);
 
 	/**
 	 * Given a property handle, try and extract the editable tag containers from within it (when dealing with a struct property of type FGameplayTagContainer).
 	 * @return True if it was possible to extract containers (even if no containers were extracted), or false otherwise.
 	 */
-	static bool GetEditableTagContainersFromPropertyHandle(const TSharedRef<IPropertyHandle>& PropHandle, TArray<FEditableGameplayTagContainerDatum>& OutEditableContainers);
+	static UE_API bool GetEditableTagContainersFromPropertyHandle(const TSharedRef<IPropertyHandle>& PropHandle, TArray<FEditableGameplayTagContainerDatum>& OutEditableContainers);
 
 	/** Construct the actual widget */
-	void Construct(const FArguments& InArgs, const TArray<FEditableGameplayTagContainerDatum>& EditableTagContainers);
+	UE_API void Construct(const FArguments& InArgs, const TArray<FEditableGameplayTagContainerDatum>& EditableTagContainers);
 
-	virtual void Tick(const FGeometry& AllottedGeometry, const double InCurrentTime, const float InDeltaTime) override;
+	UE_API virtual void Tick(const FGeometry& AllottedGeometry, const double InCurrentTime, const float InDeltaTime) override;
 
 	/** Ensures that this widget will always account for the MaxHeight if it's specified */
-	virtual FVector2D ComputeDesiredSize(float LayoutScaleMultiplier) const override;
+	UE_API virtual FVector2D ComputeDesiredSize(float LayoutScaleMultiplier) const override;
 
 	/** Updates the tag list when the filter text changes */
-	void OnFilterTextChanged( const FText& InFilterText );
+	UE_API void OnFilterTextChanged( const FText& InFilterText );
 
 	/** Returns true if this TagNode has any children that match the current filter */
-	bool FilterChildrenCheck( TSharedPtr<FGameplayTagNode>  );
+	UE_API bool FilterChildrenCheck( TSharedPtr<FGameplayTagNode>  );
 
 	/** Returns true if we're currently adding a new tag to an INI file */
-	bool IsAddingNewTag() const;
+	UE_API bool IsAddingNewTag() const;
 
 	/** Refreshes the tags that should be displayed by the widget */
-	void RefreshTags();
+	UE_API void RefreshTags();
 
 	/** Forces the widget to refresh its tags on the next tick */
-	void RefreshOnNextTick();
+	UE_API void RefreshOnNextTick();
 
 	/** Gets the widget to focus once the menu opens. */
-	TSharedPtr<SWidget> GetWidgetToFocusOnOpen();
+	UE_API TSharedPtr<SWidget> GetWidgetToFocusOnOpen();
 
 private:
 
 	/** Verify the tags are all valid and if not prompt the user. */
-	void VerifyAssetTagValidity();
+	UE_API void VerifyAssetTagValidity();
 
 	/* Filters the tree view based on the current filter text. */
-	void FilterTagTree();
+	UE_API void FilterTagTree();
 
 	/* string that sets the section of the ini file to use for this class*/ 
-	static const FString SettingsIniSection;
+	static UE_API const FString SettingsIniSection;
 
 	/* Filename of ini file to store state of the UI */
-	static const FString& GetGameplayTagsEditorStateIni();
+	static UE_API const FString& GetGameplayTagsEditorStateIni();
 
 	/* Holds the Name of this TagContainer used for saving out expansion settings */
 	FString TagContainerName;
@@ -229,7 +231,7 @@ private:
 	 * 
 	 * @return Generated row widget for the item node
 	 */
-	TSharedRef<ITableRow> OnGenerateRow(TSharedPtr<FGameplayTagNode> InItem, const TSharedRef<STableViewBase>& OwnerTable);
+	UE_API TSharedRef<ITableRow> OnGenerateRow(TSharedPtr<FGameplayTagNode> InItem, const TSharedRef<STableViewBase>& OwnerTable);
 
 	/**
 	 * Get children nodes of the specified node
@@ -237,7 +239,7 @@ private:
 	 * @param InItem		Node to get children of
 	 * @param OutChildren	[OUT] Array of children nodes, if any
 	 */
-	void OnGetChildren(TSharedPtr<FGameplayTagNode> InItem, TArray< TSharedPtr<FGameplayTagNode> >& OutChildren);
+	UE_API void OnGetChildren(TSharedPtr<FGameplayTagNode> InItem, TArray< TSharedPtr<FGameplayTagNode> >& OutChildren);
 
 	/**
 	 * Called via delegate when the status of a check box in a row changes
@@ -245,7 +247,7 @@ private:
 	 * @param NewCheckState	New check box state
 	 * @param NodeChanged	Node that was checked/unchecked
 	 */
-	void OnTagCheckStatusChanged(ECheckBoxState NewCheckState, TSharedPtr<FGameplayTagNode> NodeChanged);
+	UE_API void OnTagCheckStatusChanged(ECheckBoxState NewCheckState, TSharedPtr<FGameplayTagNode> NodeChanged);
 
 	/**
 	 * Called via delegate to determine the checkbox state of the specified node
@@ -254,12 +256,12 @@ private:
 	 * 
 	 * @return Checkbox state of the specified node
 	 */
-	ECheckBoxState IsTagChecked(TSharedPtr<FGameplayTagNode> Node) const;
+	UE_API ECheckBoxState IsTagChecked(TSharedPtr<FGameplayTagNode> Node) const;
 
 	/**
 	 * @return true if the exact Tag provided is included in any of the tag containers the widget is editing.
 	 */
-	bool IsExactTagInCollection(TSharedPtr<FGameplayTagNode> Node) const;
+	UE_API bool IsExactTagInCollection(TSharedPtr<FGameplayTagNode> Node) const;
 
 	/**
 	 * Called via delegate when the status of the allow non-restricted children check box in a row changes
@@ -267,7 +269,7 @@ private:
 	 * @param NewCheckState	New check box state
 	 * @param NodeChanged	Node that was checked/unchecked
 	 */
-	void OnAllowChildrenTagCheckStatusChanged(ECheckBoxState NewCheckState, TSharedPtr<FGameplayTagNode> NodeChanged);
+	UE_API void OnAllowChildrenTagCheckStatusChanged(ECheckBoxState NewCheckState, TSharedPtr<FGameplayTagNode> NodeChanged);
 
 	/**
 	 * Called via delegate to determine the non-restricted children checkbox state of the specified node
@@ -276,24 +278,24 @@ private:
 	 *
 	 * @return Non-restricted children heckbox state of the specified node
 	 */
-	ECheckBoxState IsAllowChildrenTagChecked(TSharedPtr<FGameplayTagNode> Node) const;
+	UE_API ECheckBoxState IsAllowChildrenTagChecked(TSharedPtr<FGameplayTagNode> Node) const;
 
 	/** Helper function to determine the visibility of the checkbox for allowing non-restricted children of restricted gameplay tags */
-	EVisibility DetermineAllowChildrenVisible(TSharedPtr<FGameplayTagNode> Node) const;
+	UE_API EVisibility DetermineAllowChildrenVisible(TSharedPtr<FGameplayTagNode> Node) const;
 
 	/**
 	 * Helper function called when the specified node is checked
 	 * 
 	 * @param NodeChecked	Node that was checked by the user
 	 */
-	void OnTagChecked(TSharedPtr<FGameplayTagNode> NodeChecked);
+	UE_API void OnTagChecked(TSharedPtr<FGameplayTagNode> NodeChecked);
 
 	/**
 	 * Helper function called when the specified node is unchecked
 	 * 
 	 * @param NodeUnchecked	Node that was unchecked by the user
 	 */
-	void OnTagUnchecked(TSharedPtr<FGameplayTagNode> NodeUnchecked);
+	UE_API void OnTagUnchecked(TSharedPtr<FGameplayTagNode> NodeUnchecked);
 
 	/**
 	 * Recursive function to uncheck all child tags
@@ -301,7 +303,7 @@ private:
 	 * @param NodeUnchecked	Node that was unchecked by the user
 	 * @param EditableContainer The container we are removing the tags from
 	 */
-	void UncheckChildren(TSharedPtr<FGameplayTagNode> NodeUnchecked, FGameplayTagContainer& EditableContainer);
+	UE_API void UncheckChildren(TSharedPtr<FGameplayTagNode> NodeUnchecked, FGameplayTagContainer& EditableContainer);
 
 	/**
 	 * Called via delegate to determine the text colour of the specified node
@@ -310,23 +312,23 @@ private:
 	 *
 	 * @return Text colour of the specified node
 	 */
-	FSlateColor GetTagTextColour(TSharedPtr<FGameplayTagNode> Node) const;
+	UE_API FSlateColor GetTagTextColour(TSharedPtr<FGameplayTagNode> Node) const;
 
 	/** Called when the user clicks the "Clear All" button; Clears all tags */
-	FReply OnClearAllClicked();
+	UE_API FReply OnClearAllClicked();
 
 	/** Called when the user clicks the "Expand All" button; Expands the entire tag tree */
-	FReply OnExpandAllClicked();
+	UE_API FReply OnExpandAllClicked();
 
 	/** Called when the user clicks the "Collapse All" button; Collapses the entire tag tree */
-	FReply OnCollapseAllClicked();
+	UE_API FReply OnCollapseAllClicked();
 
 	/**
 	 * Helper function to set the expansion state of the tree widget
 	 * 
 	 * @param bExpand If true, expand the entire tree; Otherwise, collapse the entire tree
 	 */
-	void SetTagTreeItemExpansion(bool bExpand);
+	UE_API void SetTagTreeItemExpansion(bool bExpand);
 
 	/**
 	 * Helper function to set the expansion state of a specific node
@@ -334,91 +336,93 @@ private:
 	 * @param Node		Node to set the expansion state of
 	 * @param bExapnd	If true, expand the node; Otherwise, collapse the node
 	 */
-	void SetTagNodeItemExpansion(TSharedPtr<FGameplayTagNode> Node, bool bExpand);
+	UE_API void SetTagNodeItemExpansion(TSharedPtr<FGameplayTagNode> Node, bool bExpand);
 
 	/** Load settings for the tags*/
-	void LoadSettings();
+	UE_API void LoadSettings();
 
 	/** Migrate settings */
-	void MigrateSettings();
+	UE_API void MigrateSettings();
 
 	/** Helper function to determine the visibility of the expandable UI controls */
-	EVisibility DetermineExpandableUIVisibility() const;
+	UE_API EVisibility DetermineExpandableUIVisibility() const;
 
 	/** Helper function to determine the visibility of the add new source expandable UI controls */
-	EVisibility DetermineAddNewSourceExpandableUIVisibility() const;
+	UE_API EVisibility DetermineAddNewSourceExpandableUIVisibility() const;
 
 	/** Helper function to determine the visibility of the Add New Tag widget */
-	EVisibility DetermineAddNewTagWidgetVisibility() const;
+	UE_API EVisibility DetermineAddNewTagWidgetVisibility() const;
 
 	/** Helper function to determine the visibility of the Add New Tag Source widget */
-	EVisibility DetermineAddNewSourceWidgetVisibility() const;
+	UE_API EVisibility DetermineAddNewSourceWidgetVisibility() const;
 
 	/** Helper function to determine the visibility of the tag tree controls. */
-	EVisibility DetermineTagTreeControlsVisibility() const;
+	UE_API EVisibility DetermineTagTreeControlsVisibility() const;
 
 	/** Helper function to determine the visibility of the Add New Subtag widget */
-	EVisibility DetermineAddNewSubTagWidgetVisibility(TSharedPtr<FGameplayTagNode> Node) const;
+	UE_API EVisibility DetermineAddNewSubTagWidgetVisibility(TSharedPtr<FGameplayTagNode> Node) const;
 
 	/** Helper function to determine the visibility of the Clear Selection button */
-	EVisibility DetermineClearSelectionVisibility() const;
+	UE_API EVisibility DetermineClearSelectionVisibility() const;
 
 	/** Recursive load function to go through all tags in the tree and set the expansion*/
-	void LoadTagNodeItemExpansion( TSharedPtr<FGameplayTagNode> Node );
+	UE_API void LoadTagNodeItemExpansion( TSharedPtr<FGameplayTagNode> Node );
 
 	/** Recursive function to go through all tags in the tree and set the expansion to default*/
-	void SetDefaultTagNodeItemExpansion( TSharedPtr<FGameplayTagNode> Node );
+	UE_API void SetDefaultTagNodeItemExpansion( TSharedPtr<FGameplayTagNode> Node );
 
 	/** Expansion changed callback */
-	void OnExpansionChanged( TSharedPtr<FGameplayTagNode> InItem, bool bIsExpanded );
+	UE_API void OnExpansionChanged( TSharedPtr<FGameplayTagNode> InItem, bool bIsExpanded );
 
 	/** Callback for when a new tag is added */
-	void OnGameplayTagAdded(const FString& TagName, const FString& TagComment, const FName& TagSource);
+	UE_API void OnGameplayTagAdded(const FString& TagName, const FString& TagComment, const FName& TagSource);
 
 	/** Callback when the user wants to add a subtag to an existing tag */
-	FReply OnAddSubtagClicked(TSharedPtr<FGameplayTagNode> InTagNode);
+	UE_API FReply OnAddSubtagClicked(TSharedPtr<FGameplayTagNode> InTagNode);
 
 	/** Creates a dropdown menu to provide additional functionality for tags (renaming, deletion, search for references, etc.) */
-	TSharedRef<SWidget> MakeTagActionsMenu(TSharedPtr<FGameplayTagNode> InTagNode);
+	UE_API TSharedRef<SWidget> MakeTagActionsMenu(TSharedPtr<FGameplayTagNode> InTagNode);
 
 	/** Attempts to rename the tag through a dialog box */
-	void OnRenameTag(TSharedPtr<FGameplayTagNode> InTagNode);
+	UE_API void OnRenameTag(TSharedPtr<FGameplayTagNode> InTagNode);
 
 	/** Attempts to delete the specified tag */
-	void OnDeleteTag(TSharedPtr<FGameplayTagNode> InTagNode);
+	UE_API void OnDeleteTag(TSharedPtr<FGameplayTagNode> InTagNode);
 
 	/** Attempts to add the exact specified tag*/
-	void OnAddTag(TSharedPtr<FGameplayTagNode> InTagNode);
+	UE_API void OnAddTag(TSharedPtr<FGameplayTagNode> InTagNode);
 
 	/** Attempts to remove the specified tag, but not the children */
-	void OnRemoveTag(TSharedPtr<FGameplayTagNode> InTagNode);
+	UE_API void OnRemoveTag(TSharedPtr<FGameplayTagNode> InTagNode);
 
 	/** Searches for all references for the selected tag */
-	void OnSearchForReferences(TSharedPtr<FGameplayTagNode> InTagNode);
+	UE_API void OnSearchForReferences(TSharedPtr<FGameplayTagNode> InTagNode);
 
 	/** Returns true if the user can select tags from the widget */
-	bool CanSelectTags() const;
+	UE_API bool CanSelectTags() const;
 
 	/** Determines if the expandable UI that contains the Add New Tag widget should be expanded or collapsed */
-	ECheckBoxState GetAddTagSectionExpansionState() const;
+	UE_API ECheckBoxState GetAddTagSectionExpansionState() const;
 
 	/** Callback for when the state of the expandable UI section changes */
-	void OnAddTagSectionExpansionStateChanged(ECheckBoxState NewState);
+	UE_API void OnAddTagSectionExpansionStateChanged(ECheckBoxState NewState);
 
 	/** Determines if the expandable UI that contains the Add New Tag Source widget should be expanded or collapsed */
-	ECheckBoxState GetAddSourceSectionExpansionState() const;
+	UE_API ECheckBoxState GetAddSourceSectionExpansionState() const;
 
 	/** Callback for when the state of the expandable tag source UI section changes */
-	void OnAddSourceSectionExpansionStateChanged(ECheckBoxState NewState);
+	UE_API void OnAddSourceSectionExpansionStateChanged(ECheckBoxState NewState);
 
-	void SetContainer(FGameplayTagContainer* OriginalContainer, FGameplayTagContainer* EditedContainer, UObject* OwnerObj);
+	UE_API void SetContainer(FGameplayTagContainer* OriginalContainer, FGameplayTagContainer* EditedContainer, UObject* OwnerObj);
 
 	/** Opens a dialog window to rename the selected tag */
-	void OpenRenameGameplayTagDialog(TSharedPtr<FGameplayTagNode> GameplayTagNode) const;
+	UE_API void OpenRenameGameplayTagDialog(TSharedPtr<FGameplayTagNode> GameplayTagNode) const;
 
 	/** Delegate that is fired when a tag is successfully renamed */
-	void OnGameplayTagRenamed(FString OldTagName, FString NewTagName);
+	UE_API void OnGameplayTagRenamed(FString OldTagName, FString NewTagName);
 	
 	/** Populate tag items from the gameplay tags manager. */
-	void GetFilteredGameplayRootTags(const FString& InFilterString, TArray<TSharedPtr<FGameplayTagNode>>& OutNodes) const;
+	UE_API void GetFilteredGameplayRootTags(const FString& InFilterString, TArray<TSharedPtr<FGameplayTagNode>>& OutNodes) const;
 };
+
+#undef UE_API
diff --git a/Engine/Plugins/Editor/MeshLODToolset/Source/MeshLODToolset/Public/Tools/GenerateStaticMeshLODAssetTool.h b/Engine/Plugins/Editor/MeshLODToolset/Source/MeshLODToolset/Public/Tools/GenerateStaticMeshLODAssetTool.h
--- a/Engine/Plugins/Editor/MeshLODToolset/Source/MeshLODToolset/Public/Tools/GenerateStaticMeshLODAssetTool.h
+++ b/Engine/Plugins/Editor/MeshLODToolset/Source/MeshLODToolset/Public/Tools/GenerateStaticMeshLODAssetTool.h
@@ -7,6 +7,8 @@
 #include "Graphs/GenerateStaticMeshLODProcess.h"
 #include "GenerateStaticMeshLODAssetTool.generated.h"
 
+#define UE_API MESHLODTOOLSET_API
+
 class UCollisionGeometryVisualizationProperties;
 class UPhysicsObjectToolPropertySet;
 class UPreviewGeometry;
@@ -35,8 +37,8 @@ enum class EGenerateLODAssetOutputMode : uint8
 /**
  * Tool builder
  */
-UCLASS()
-class MESHLODTOOLSET_API UGenerateStaticMeshLODAssetToolBuilder : public UMultiSelectionMeshEditingToolBuilder
+UCLASS(MinimalAPI)
+class UGenerateStaticMeshLODAssetToolBuilder : public UMultiSelectionMeshEditingToolBuilder
 {
 	GENERATED_BODY()
 
@@ -44,17 +46,17 @@ public:
 	bool bUseAssetEditorMode = false;
 	bool bInRestrictiveMode = false;
 
-	virtual bool CanBuildTool(const FToolBuilderState& SceneState) const override;
-	virtual UMultiSelectionMeshEditingTool* CreateNewTool(const FToolBuilderState& SceneState) const override;
+	UE_API virtual bool CanBuildTool(const FToolBuilderState& SceneState) const override;
+	UE_API virtual UMultiSelectionMeshEditingTool* CreateNewTool(const FToolBuilderState& SceneState) const override;
 
 protected:
-	virtual const FToolTargetTypeRequirements& GetTargetRequirements() const override;
+	UE_API virtual const FToolTargetTypeRequirements& GetTargetRequirements() const override;
 };
 
 
 
-UCLASS()
-class MESHLODTOOLSET_API UGenerateStaticMeshLODAssetToolOutputProperties : public UInteractiveToolPropertySet
+UCLASS(MinimalAPI)
+class UGenerateStaticMeshLODAssetToolOutputProperties : public UInteractiveToolPropertySet
 {
 	GENERATED_BODY()
 public:
@@ -90,15 +92,15 @@ enum class EGenerateLODAssetToolPresetAction : uint8
 };
 
 
-UCLASS()
-class MESHLODTOOLSET_API UGenerateStaticMeshLODAssetToolPresetProperties : public UInteractiveToolPropertySet
+UCLASS(MinimalAPI)
+class UGenerateStaticMeshLODAssetToolPresetProperties : public UInteractiveToolPropertySet
 {
 	GENERATED_BODY()
 
 public:
 	TWeakObjectPtr<UGenerateStaticMeshLODAssetTool> ParentTool;
 	void Initialize(UGenerateStaticMeshLODAssetTool* ParentToolIn) { ParentTool = ParentToolIn; }
-	virtual void PostAction(EGenerateLODAssetToolPresetAction Action);
+	UE_API virtual void PostAction(EGenerateLODAssetToolPresetAction Action);
 
 public:
 	/** Preset Asset represents a set of Saved Settings for this Tool */
@@ -125,8 +127,8 @@ public:
 /**
  * Standard properties
  */
-UCLASS()
-class MESHLODTOOLSET_API UGenerateStaticMeshLODAssetToolProperties : public UInteractiveToolPropertySet
+UCLASS(MinimalAPI)
+class UGenerateStaticMeshLODAssetToolProperties : public UInteractiveToolPropertySet
 {
 	GENERATED_BODY()
 
@@ -201,7 +203,7 @@ enum class EGenerateStaticMeshLOD_BakeConstraint
 };
 
 USTRUCT()
-struct MESHLODTOOLSET_API FGenerateStaticMeshLOD_TextureConfig
+struct FGenerateStaticMeshLOD_TextureConfig
 {
 	GENERATED_BODY()
 
@@ -219,7 +221,7 @@ struct MESHLODTOOLSET_API FGenerateStaticMeshLOD_TextureConfig
 
 
 USTRUCT()
-struct MESHLODTOOLSET_API FGenerateStaticMeshLOD_MaterialConfig
+struct FGenerateStaticMeshLOD_MaterialConfig
 {
 	GENERATED_BODY()
 
@@ -237,8 +239,8 @@ struct MESHLODTOOLSET_API FGenerateStaticMeshLOD_MaterialConfig
 
 
 
-UCLASS()
-class MESHLODTOOLSET_API UGenerateStaticMeshLODAssetToolTextureProperties : public UInteractiveToolPropertySet
+UCLASS(MinimalAPI)
+class UGenerateStaticMeshLODAssetToolTextureProperties : public UInteractiveToolPropertySet
 {
 	GENERATED_BODY()
 public:
@@ -258,29 +260,29 @@ public:
 /**
  * Simple tool to combine multiple meshes into a single mesh asset
  */
-UCLASS()
-class MESHLODTOOLSET_API UGenerateStaticMeshLODAssetTool : public UMultiSelectionMeshEditingTool
+UCLASS(MinimalAPI)
+class UGenerateStaticMeshLODAssetTool : public UMultiSelectionMeshEditingTool
 {
 	GENERATED_BODY()
 
 public:
 
 	// Enable UI Customization for running this Tool in the Static Mesh Asset Editor. Must be called before Setup.
-	virtual void SetUseAssetEditorMode(bool bEnable);
+	UE_API virtual void SetUseAssetEditorMode(bool bEnable);
 
 	// Enable restrictive mode. Must be called before Setup.
-	virtual void SetRestrictiveMode(bool bEnable);
+	UE_API virtual void SetRestrictiveMode(bool bEnable);
 
-	virtual void Setup() override;
-	virtual void OnShutdown(EToolShutdownType ShutdownType) override;
+	UE_API virtual void Setup() override;
+	UE_API virtual void OnShutdown(EToolShutdownType ShutdownType) override;
 
-	virtual void OnTick(float DeltaTime);
+	UE_API virtual void OnTick(float DeltaTime);
 
 	virtual bool HasCancel() const override { return true; }
 	virtual bool HasAccept() const override { return true; }
-	virtual bool CanAccept() const override;
+	UE_API virtual bool CanAccept() const override;
 
-	virtual void RequestPresetAction(EGenerateLODAssetToolPresetAction ActionType);
+	UE_API virtual void RequestPresetAction(EGenerateLODAssetToolPresetAction ActionType);
 
 protected:
 
@@ -332,12 +334,14 @@ protected:
 
 	TUniquePtr<UE::Geometry::IDynamicMeshOperatorFactory> OpFactory;
 
-	void OnSettingsModified();
+	UE_API void OnSettingsModified();
 
-	void CreateNewAsset();
-	void UpdateExistingAsset();
+	UE_API void CreateNewAsset();
+	UE_API void UpdateExistingAsset();
 
-	void OnPresetSelectionChanged();
+	UE_API void OnPresetSelectionChanged();
 
-	bool ValidateSettings() const;
+	UE_API bool ValidateSettings() const;
 };
+
+#undef UE_API
diff --git a/Engine/Plugins/Editor/MeshLODToolset/Source/MeshLODToolset/Public/Tools/LODGenerationSettingsAsset.h b/Engine/Plugins/Editor/MeshLODToolset/Source/MeshLODToolset/Public/Tools/LODGenerationSettingsAsset.h
--- a/Engine/Plugins/Editor/MeshLODToolset/Source/MeshLODToolset/Public/Tools/LODGenerationSettingsAsset.h
+++ b/Engine/Plugins/Editor/MeshLODToolset/Source/MeshLODToolset/Public/Tools/LODGenerationSettingsAsset.h
@@ -15,8 +15,8 @@
  * The Tool uses these serialized settings as a 'Preset', ie the user can save a set
  * of configured settings, or load previously-saved settings. 
  */
-UCLASS(Blueprintable)
-class MESHLODTOOLSET_API UStaticMeshLODGenerationSettings : public UObject
+UCLASS(MinimalAPI, Blueprintable)
+class UStaticMeshLODGenerationSettings : public UObject
 {
 	GENERATED_BODY()
 
diff --git a/Engine/Plugins/Editor/MeshLODToolset/Source/MeshLODToolset/Public/Tools/LODManagerTool.h b/Engine/Plugins/Editor/MeshLODToolset/Source/MeshLODToolset/Public/Tools/LODManagerTool.h
--- a/Engine/Plugins/Editor/MeshLODToolset/Source/MeshLODToolset/Public/Tools/LODManagerTool.h
+++ b/Engine/Plugins/Editor/MeshLODToolset/Source/MeshLODToolset/Public/Tools/LODManagerTool.h
@@ -7,6 +7,8 @@
 #include "Engine/StaticMesh.h"
 #include "LODManagerTool.generated.h"
 
+#define UE_API MESHLODTOOLSET_API
+
 class UPreviewGeometry;
 class UPreviewMesh;
 
@@ -17,22 +19,22 @@ class ULODManagerTool;
 /**
  *
  */
-UCLASS()
-class MESHLODTOOLSET_API ULODManagerToolBuilder : public UMultiSelectionMeshEditingToolBuilder
+UCLASS(MinimalAPI)
+class ULODManagerToolBuilder : public UMultiSelectionMeshEditingToolBuilder
 {
 	GENERATED_BODY()
 
 public:
-	virtual bool CanBuildTool(const FToolBuilderState& SceneState) const override;
-	virtual UMultiSelectionMeshEditingTool* CreateNewTool(const FToolBuilderState& SceneState) const override;
+	UE_API virtual bool CanBuildTool(const FToolBuilderState& SceneState) const override;
+	UE_API virtual UMultiSelectionMeshEditingTool* CreateNewTool(const FToolBuilderState& SceneState) const override;
 
 protected:
-	virtual const FToolTargetTypeRequirements& GetTargetRequirements() const override;
+	UE_API virtual const FToolTargetTypeRequirements& GetTargetRequirements() const override;
 };
 
 
 USTRUCT()
-struct MESHLODTOOLSET_API FLODManagerLODInfo
+struct FLODManagerLODInfo
 {
 	GENERATED_BODY()
 
@@ -45,8 +47,8 @@ struct MESHLODTOOLSET_API FLODManagerLODInfo
 
 
 
-UCLASS()
-class MESHLODTOOLSET_API ULODManagerLODProperties : public UInteractiveToolPropertySet
+UCLASS(MinimalAPI)
+class ULODManagerLODProperties : public UInteractiveToolPropertySet
 {
 	GENERATED_BODY()
 public:
@@ -74,8 +76,8 @@ public:
 
 
 
-UCLASS()
-class MESHLODTOOLSET_API ULODManagerPreviewLODProperties : public UInteractiveToolPropertySet
+UCLASS(MinimalAPI)
+class ULODManagerPreviewLODProperties : public UInteractiveToolPropertySet
 {
 	GENERATED_BODY()
 public:
@@ -113,8 +115,8 @@ enum class ELODManagerToolActions
 
 
 
-UCLASS()
-class MESHLODTOOLSET_API ULODManagerActionPropertySet : public UInteractiveToolPropertySet
+UCLASS(MinimalAPI)
+class ULODManagerActionPropertySet : public UInteractiveToolPropertySet
 {
 	GENERATED_BODY()
 
@@ -122,13 +124,13 @@ public:
 	TWeakObjectPtr<ULODManagerTool> ParentTool;
 
 	void Initialize(ULODManagerTool* ParentToolIn) { ParentTool = ParentToolIn; }
-	void PostAction(ELODManagerToolActions Action);
+	UE_API void PostAction(ELODManagerToolActions Action);
 
 };
 
 
-UCLASS()
-class MESHLODTOOLSET_API ULODManagerHiResSourceModelActions : public ULODManagerActionPropertySet
+UCLASS(MinimalAPI)
+class ULODManagerHiResSourceModelActions : public ULODManagerActionPropertySet
 {
 	GENERATED_BODY()
 public:
@@ -148,8 +150,8 @@ public:
 };
 
 
-UCLASS()
-class MESHLODTOOLSET_API ULODManagerMaterialActions : public ULODManagerActionPropertySet
+UCLASS(MinimalAPI)
+class ULODManagerMaterialActions : public ULODManagerActionPropertySet
 {
 	GENERATED_BODY()
 public:
@@ -188,24 +190,24 @@ public:
 /**
  * Mesh Attribute Editor Tool
  */
-UCLASS()
-class MESHLODTOOLSET_API ULODManagerTool : public UMultiSelectionMeshEditingTool, public ILODManagerToolChangeTarget
+UCLASS(MinimalAPI)
+class ULODManagerTool : public UMultiSelectionMeshEditingTool, public ILODManagerToolChangeTarget
 {
 	GENERATED_BODY()
 
 public:
-	ULODManagerTool();
+	UE_API ULODManagerTool();
 
-	virtual void Setup() override;
-	virtual void OnShutdown(EToolShutdownType ShutdownType) override;
-	virtual void OnTick(float DeltaTime) override;
+	UE_API virtual void Setup() override;
+	UE_API virtual void OnShutdown(EToolShutdownType ShutdownType) override;
+	UE_API virtual void OnTick(float DeltaTime) override;
 
 	virtual bool HasCancel() const override { return true; }
 	virtual bool HasAccept() const override { return true; }
 
-	virtual void RequestAction(ELODManagerToolActions ActionType);
+	UE_API virtual void RequestAction(ELODManagerToolActions ActionType);
 
-	virtual void ApplyChange(const UE::Geometry::LODManagerHelper::FLODManagerToolChange* Change, bool bRevert) override;
+	UE_API virtual void ApplyChange(const UE::Geometry::LODManagerHelper::FLODManagerToolChange* Change, bool bRevert) override;
 
 protected:
 
@@ -225,13 +227,13 @@ protected:
 
 public:
 	UFUNCTION()
-	void DeleteHiResSourceModel();
+	UE_API void DeleteHiResSourceModel();
 
 	UFUNCTION()
-	void MoveHiResToLOD0();
+	UE_API void MoveHiResToLOD0();
 
 	UFUNCTION()
-	void RemoveUnreferencedMaterials();
+	UE_API void RemoveUnreferencedMaterials();
 
 
 	struct FLODName
@@ -250,18 +252,18 @@ public:
 
 
 protected:
-	UStaticMesh* GetSingleStaticMesh();
+	UE_API UStaticMesh* GetSingleStaticMesh();
 
 	ELODManagerToolActions PendingAction = ELODManagerToolActions::NoAction;
 	
 
 	bool bLODInfoValid = false;
 	// captures the material list and triangle and vertex counts for the current configuration of lods
-	void UpdateLODInfo();
+	UE_API void UpdateLODInfo();
 
 	// maps pretty name in UI to description of the LOD
 	TMap<FString, FLODName> ActiveLODNames;
-	void UpdateLODNames();
+	UE_API void UpdateLODNames();
 
 	UPROPERTY()
 	TObjectPtr<UPreviewMesh> LODPreview;
@@ -270,17 +272,17 @@ protected:
 	TObjectPtr<UPreviewGeometry> LODPreviewLines;
 
 	bool bPreviewLODValid = false;
-	void UpdatePreviewLOD();
-	void UpdatePreviewLines(FLODMeshInfo& LODMeshInfo);
-	void ClearPreviewLines();
+	UE_API void UpdatePreviewLOD();
+	UE_API void UpdatePreviewLines(FLODMeshInfo& LODMeshInfo);
+	UE_API void ClearPreviewLines();
 
 	// returns the requested LOD as a dynamic mesh along with boundary edges.  Has internal cache
-	TUniquePtr<ULODManagerTool::FLODMeshInfo> GetLODMeshInfo(const FLODName& LODName);
+	UE_API TUniquePtr<ULODManagerTool::FLODMeshInfo> GetLODMeshInfo(const FLODName& LODName);
 
 	// For undo / redo system with our custom changes.  
 	TUniquePtr< UE::Geometry::LODManagerHelper::FLODManagerToolChange > ActiveChange;
-	void BeginChange(FText TransactionName);
-	void EndChange();
+	UE_API void BeginChange(FText TransactionName);
+	UE_API void EndChange();
 
 	// state information used for undo within the tool
 	TUniquePtr< UE::Geometry::LODManagerHelper::FProxyLODState>  ProxyLODState;
@@ -288,3 +290,5 @@ protected:
 	// cache of dynamic mesh representation for each lod, imported and renderdata.
 	TUniquePtr< UE::Geometry::LODManagerHelper::FDynamicMeshLODCache > DynamicMeshCache;
 };
+
+#undef UE_API
diff --git a/Engine/Plugins/Editor/ModelingToolsEditorMode/Source/ModelingToolsEditorMode/Public/ModelingModeToolExtensions.h b/Engine/Plugins/Editor/ModelingToolsEditorMode/Source/ModelingToolsEditorMode/Public/ModelingModeToolExtensions.h
--- a/Engine/Plugins/Editor/ModelingToolsEditorMode/Source/ModelingToolsEditorMode/Public/ModelingModeToolExtensions.h
+++ b/Engine/Plugins/Editor/ModelingToolsEditorMode/Source/ModelingToolsEditorMode/Public/ModelingModeToolExtensions.h
@@ -30,7 +30,7 @@ public:
  * Note that if bIsInfoQueryOnly is true, then the Extension does not need to 
  * construct Tool Builders, the query is for information purposes only
  */
-struct MODELINGTOOLSEDITORMODE_API FExtensionToolQueryInfo
+struct FExtensionToolQueryInfo
 {
 	bool bIsInfoQueryOnly = false;
 
@@ -43,7 +43,7 @@ struct MODELINGTOOLSEDITORMODE_API FExtensionToolQueryInfo
  * IModelingModeToolExtension implementations return the list of Tools they provide
  * via instances of FExtensionToolDescription.
  */
-struct MODELINGTOOLSEDITORMODE_API FExtensionToolDescription
+struct FExtensionToolDescription
 {
 	/** Long name of the Tool, used in various places in the UI */
 	FText ToolName;
@@ -63,7 +63,7 @@ struct MODELINGTOOLSEDITORMODE_API FExtensionToolDescription
  * IModelingModeToolExtension implementations can optionally return additional information used in the Eidtor UI
  * via instances of FModelingModeExtensionExtendedInfo.
  */
-struct MODELINGTOOLSEDITORMODE_API FModelingModeExtensionExtendedInfo
+struct FModelingModeExtensionExtendedInfo
 {
 	/** Tooltip to use for UI buttons that refer to the Extension */
 	FText ToolPaletteButtonTooltip;
@@ -77,7 +77,7 @@ struct MODELINGTOOLSEDITORMODE_API FModelingModeExtensionExtendedInfo
  * included in a section of the Modeling Mode tool list, based on GetToolSectionName().
  * 
  */
-class MODELINGTOOLSEDITORMODE_API IModelingModeToolExtension : public IModularFeature
+class IModelingModeToolExtension : public IModularFeature
 {
 public:
 	virtual ~IModelingModeToolExtension() {}
diff --git a/Engine/Plugins/Editor/ModelingToolsEditorMode/Source/ModelingToolsEditorMode/Public/ModelingToolsEditorModeSettings.h b/Engine/Plugins/Editor/ModelingToolsEditorMode/Source/ModelingToolsEditorMode/Public/ModelingToolsEditorModeSettings.h
--- a/Engine/Plugins/Editor/ModelingToolsEditorMode/Source/ModelingToolsEditorMode/Public/ModelingToolsEditorModeSettings.h
+++ b/Engine/Plugins/Editor/ModelingToolsEditorMode/Source/ModelingToolsEditorMode/Public/ModelingToolsEditorModeSettings.h
@@ -5,6 +5,8 @@
 #include "Engine/DeveloperSettings.h"
 #include "ModelingToolsEditorModeSettings.generated.h"
 
+#define UE_API MODELINGTOOLSEDITORMODE_API
+
 struct FCollectionReference;
 
 
@@ -54,8 +56,8 @@ enum class EModelingModeAssetGenerationLocation
 /**
  * Settings for the Modeling Tools Editor Mode plug-in.
  */
-UCLASS(config=Editor)
-class MODELINGTOOLSEDITORMODE_API UModelingToolsEditorModeSettings : public UDeveloperSettings
+UCLASS(MinimalAPI, config=Editor)
+class UModelingToolsEditorModeSettings : public UDeveloperSettings
 {
 	GENERATED_BODY()
 
@@ -67,8 +69,8 @@ public:
 	virtual FName GetCategoryName() const override { return FName("Plugins"); }
 	virtual FName GetSectionName() const override { return FName("ModelingMode"); }
 
-	virtual FText GetSectionText() const override;
-	virtual FText GetSectionDescription() const override;
+	UE_API virtual FText GetSectionText() const override;
+	UE_API virtual FText GetSectionDescription() const override;
 
 protected:
 
@@ -291,8 +293,8 @@ struct FModelingModeAssetCollectionSet
 };
 
 
-UCLASS(config=Editor)
-class MODELINGTOOLSEDITORMODE_API UModelingToolsModeCustomizationSettings : public UDeveloperSettings
+UCLASS(MinimalAPI, config=Editor)
+class UModelingToolsModeCustomizationSettings : public UDeveloperSettings
 {
 	GENERATED_BODY()
 
@@ -303,8 +305,8 @@ public:
 	virtual FName GetCategoryName() const { return FName("Plugins"); }
 	virtual FName GetSectionName() const { return FName("ModelingMode"); }
 
-	virtual FText GetSectionText() const override;
-	virtual FText GetSectionDescription() const override;
+	UE_API virtual FText GetSectionText() const override;
+	UE_API virtual FText GetSectionDescription() const override;
 
 public:
 
@@ -410,3 +412,5 @@ public:
 
 
 };
+
+#undef UE_API
diff --git a/Engine/Plugins/Editor/ModelingToolsEditorMode/Source/ModelingToolsEditorMode/Public/ModelingToolsEditorModeStyle.h b/Engine/Plugins/Editor/ModelingToolsEditorMode/Source/ModelingToolsEditorMode/Public/ModelingToolsEditorModeStyle.h
--- a/Engine/Plugins/Editor/ModelingToolsEditorMode/Source/ModelingToolsEditorMode/Public/ModelingToolsEditorModeStyle.h
+++ b/Engine/Plugins/Editor/ModelingToolsEditorMode/Source/ModelingToolsEditorMode/Public/ModelingToolsEditorModeStyle.h
@@ -4,25 +4,29 @@
 
 #include "Templates/SharedPointer.h"
 
+#define UE_API MODELINGTOOLSEDITORMODE_API
+
 struct FSlateBrush;
 
-class MODELINGTOOLSEDITORMODE_API FModelingToolsEditorModeStyle
+class FModelingToolsEditorModeStyle
 {
 public:
-	static void Initialize();
+	static UE_API void Initialize();
 
-	static void Shutdown();
+	static UE_API void Shutdown();
 
-	static TSharedPtr< class ISlateStyle > Get();
+	static UE_API TSharedPtr< class ISlateStyle > Get();
 
-	static FName GetStyleSetName();
+	static UE_API FName GetStyleSetName();
 
 	// use to access icons defined by the style set by name, eg GetBrush("BrushFalloffIcons.Smooth")
-	static const FSlateBrush* GetBrush(FName PropertyName, const ANSICHAR* Specifier = NULL);
+	static UE_API const FSlateBrush* GetBrush(FName PropertyName, const ANSICHAR* Specifier = NULL);
 
 private:
-	static FString InContent(const FString& RelativePath, const ANSICHAR* Extension);
+	static UE_API FString InContent(const FString& RelativePath, const ANSICHAR* Extension);
 
 private:
-	static TSharedPtr< class FSlateStyleSet > StyleSet;
+	static UE_API TSharedPtr< class FSlateStyleSet > StyleSet;
 };
+
+#undef UE_API
diff --git a/Engine/Plugins/Editor/ObjectMixer/LightMixer/Source/LightMixer/Public/LightMixerEditorSettings.h b/Engine/Plugins/Editor/ObjectMixer/LightMixer/Source/LightMixer/Public/LightMixerEditorSettings.h
--- a/Engine/Plugins/Editor/ObjectMixer/LightMixer/Source/LightMixer/Public/LightMixerEditorSettings.h
+++ b/Engine/Plugins/Editor/ObjectMixer/LightMixer/Source/LightMixer/Public/LightMixerEditorSettings.h
@@ -4,8 +4,8 @@
 
 #include "LightMixerEditorSettings.generated.h"
 
-UCLASS(config = ObjectMixer)
-class LIGHTMIXER_API ULightMixerEditorSettings : public UObject
+UCLASS(MinimalAPI, config = ObjectMixer)
+class ULightMixerEditorSettings : public UObject
 {
 	GENERATED_BODY()
 public:
diff --git a/Engine/Plugins/Editor/ObjectMixer/LightMixer/Source/LightMixer/Public/LightMixerModule.h b/Engine/Plugins/Editor/ObjectMixer/LightMixer/Source/LightMixer/Public/LightMixerModule.h
--- a/Engine/Plugins/Editor/ObjectMixer/LightMixer/Source/LightMixer/Public/LightMixerModule.h
+++ b/Engine/Plugins/Editor/ObjectMixer/LightMixer/Source/LightMixer/Public/LightMixerModule.h
@@ -4,25 +4,29 @@
 
 #include "ObjectMixerEditorModule.h"
 
-class LIGHTMIXER_API FLightMixerModule : public FObjectMixerEditorModule
+#define UE_API LIGHTMIXER_API
+
+class FLightMixerModule : public FObjectMixerEditorModule
 {
 public:
 
 	//~ Begin IModuleInterface Interface
-	virtual void StartupModule() override;
-	virtual void ShutdownModule() override;
+	UE_API virtual void StartupModule() override;
+	UE_API virtual void ShutdownModule() override;
 	//~ End IModuleInterface Interface
 
-	static FLightMixerModule& Get();
+	static UE_API FLightMixerModule& Get();
 	
-	static void OpenProjectSettings();
+	static UE_API void OpenProjectSettings();
 	
 	//~ Begin FObjectMixerEditorModule overrides
-	virtual void Initialize() override;
-	virtual FName GetModuleName() const override;
-	virtual void SetupMenuItemVariables() override;
-	virtual FName GetTabSpawnerId() override;
-	virtual void RegisterSettings() const override;
-	virtual void UnregisterSettings() const override;
+	UE_API virtual void Initialize() override;
+	UE_API virtual FName GetModuleName() const override;
+	UE_API virtual void SetupMenuItemVariables() override;
+	UE_API virtual FName GetTabSpawnerId() override;
+	UE_API virtual void RegisterSettings() const override;
+	UE_API virtual void UnregisterSettings() const override;
 	//~ End FObjectMixerEditorModule overrides
 };
+
+#undef UE_API
diff --git a/Engine/Plugins/Editor/ObjectMixer/LightMixer/Source/LightMixer/Public/LightMixerObjectFilter.h b/Engine/Plugins/Editor/ObjectMixer/LightMixer/Source/LightMixer/Public/LightMixerObjectFilter.h
--- a/Engine/Plugins/Editor/ObjectMixer/LightMixer/Source/LightMixer/Public/LightMixerObjectFilter.h
+++ b/Engine/Plugins/Editor/ObjectMixer/LightMixer/Source/LightMixer/Public/LightMixerObjectFilter.h
@@ -9,8 +9,8 @@
 
 #include "LightMixerObjectFilter.generated.h"
 
-UCLASS(BlueprintType, EditInlineNew)
-class LIGHTMIXER_API ULightMixerObjectFilter : public UObjectMixerObjectFilter
+UCLASS(MinimalAPI, BlueprintType, EditInlineNew)
+class ULightMixerObjectFilter : public UObjectMixerObjectFilter
 {
 	GENERATED_BODY()
 public:
diff --git a/Engine/Plugins/Editor/ObjectMixer/ObjectMixer/Source/ObjectMixer/Public/ObjectFilter/ObjectMixerEditorObjectFilter.h b/Engine/Plugins/Editor/ObjectMixer/ObjectMixer/Source/ObjectMixer/Public/ObjectFilter/ObjectMixerEditorObjectFilter.h
--- a/Engine/Plugins/Editor/ObjectMixer/ObjectMixer/Source/ObjectMixer/Public/ObjectFilter/ObjectMixerEditorObjectFilter.h
+++ b/Engine/Plugins/Editor/ObjectMixer/ObjectMixer/Source/ObjectMixer/Public/ObjectFilter/ObjectMixerEditorObjectFilter.h
@@ -8,6 +8,8 @@
 
 #include "ObjectMixerEditorObjectFilter.generated.h"
 
+#define UE_API OBJECTMIXEREDITOR_API
+
 struct FToolMenuContext;
 
 UENUM(BlueprintType)
@@ -46,8 +48,8 @@ enum class EObjectMixerTreeViewMode : uint8
  * Native class for filtering object types to Object Mixer.
  * Native C++ classes should inherit directly from this class.
  */
-UCLASS(Abstract, BlueprintType) 
-class OBJECTMIXEREDITOR_API UObjectMixerObjectFilter : public UObject
+UCLASS(MinimalAPI, Abstract, BlueprintType)
+class UObjectMixerObjectFilter : public UObject
 {
 	GENERATED_BODY()
 	
@@ -56,7 +58,7 @@ public:
 	UObjectMixerObjectFilter() = default;
 
 	/** Begin UObject overrides */
-	virtual void PostCDOCompiled(const FPostCDOCompiledContext& Context) override;
+	UE_API virtual void PostCDOCompiled(const FPostCDOCompiledContext& Context) override;
 	/** End UObject overrides */
 	
 	/**
@@ -77,31 +79,31 @@ public:
 	/**
 	 * Determines if transient objects (such as Sequencer Spawnables) should be shown in the list. False by default.
 	 */
-	virtual bool GetShowTransientObjects() const;
+	UE_API virtual bool GetShowTransientObjects() const;
 
 	/**
 	 * Determines if hybrid rows should be used when enabled in user settings. True by default.
 	 */
-	virtual bool ShouldAllowHybridRows() const;
+	UE_API virtual bool ShouldAllowHybridRows() const;
 
 	/**
 	 * Determines if the user should be able to change which columns are displayed. True by default.
 	 */
-	virtual bool ShouldAllowColumnCustomizationByUser() const;
+	UE_API virtual bool ShouldAllowColumnCustomizationByUser() const;
 
 	/**
 	 * Specify a list of property names corresponding to columns you want to show by default.
 	 * For example, you can specify "Intensity" and "LightColor" to show only those property columns by default in the UI.
 	 * Columns not specified will not be shown by default but can be enabled by the user in the UI.
 	 */
-	virtual TSet<FName> GetColumnsToShowByDefault() const;
+	UE_API virtual TSet<FName> GetColumnsToShowByDefault() const;
 
 	/**
 	 * Specify a list of property names corresponding to columns you don't want to ever show.
 	 * For example, you can specify "Intensity" and "LightColor" to ensure that they can't be enabled or shown in the UI.
 	 * Columns not specified can be enabled by the user in the UI.
 	 */
-	virtual TSet<FName> GetColumnsToExclude() const;
+	UE_API virtual TSet<FName> GetColumnsToExclude() const;
 
 	/**
 	 * Specify a list of property names found in parent classes you want to show that aren't in the specified classes.
@@ -120,55 +122,55 @@ public:
 	 * Defaults to 'None' which only considers the properties of the specified classes.
 	 * If you're not seeing all the properties you expected, try overloading this function.
 	 */
-	virtual EObjectMixerInheritanceInclusionOptions GetObjectMixerPropertyInheritanceInclusionOptions() const;
+	UE_API virtual EObjectMixerInheritanceInclusionOptions GetObjectMixerPropertyInheritanceInclusionOptions() const;
 
 	/**
 	 * Specify whether we should return only the specified classes or the parent and child classes in placement mode.
 	 * Defaults to 'None' which only considers the specified classes.
 	 */
-	virtual EObjectMixerInheritanceInclusionOptions GetObjectMixerPlacementClassInclusionOptions() const;
+	UE_API virtual EObjectMixerInheritanceInclusionOptions GetObjectMixerPlacementClassInclusionOptions() const;
 
 	/**
 	 * If true, properties that are not visible in the details panel and properties not supported by SSingleProperty will be selectable.
 	 * Defaults to false.
 	 */
-	virtual bool ShouldIncludeUnsupportedProperties() const;
+	UE_API virtual bool ShouldIncludeUnsupportedProperties() const;
 
 	/**
 	 * If a property is changed that has a name found in this set, the panel will be refreshed.
 	 * Add a property name to this list if you expect the list to change in some way after changing that property.
 	 */
-	virtual TSet<FName> GetPropertiesThatRequireListRefresh() const;
+	UE_API virtual TSet<FName> GetPropertiesThatRequireListRefresh() const;
 
 	/**
 	 * Fills OutAssociatedActors with additional actors to include as children in the ObjectMixer hierarchy for the given actor.
 	 */
-	virtual TArray<AActor*> FindAssociatedActors(AActor* InActor) const;
+	UE_API virtual TArray<AActor*> FindAssociatedActors(AActor* InActor) const;
 
 	/**
 	 * Returns true if AssociatedActor is associated with Actor.
 	 * This should return true if the actor would be included in FindAssociatedActors, but is a separate function to avoid allocating an array every time.
 	 */
-	virtual bool IsActorAssociated(AActor* Actor, AActor* AssociatedActor) const;
+	UE_API virtual bool IsActorAssociated(AActor* Actor, AActor* AssociatedActor) const;
 
 	/** 
 	 * Returns true if the given tree item should override the Object Mixer's standard handling for drop operations (ValidateDrop/OnDrop)
 	 */
-	virtual bool HasCustomDropHandling(const ISceneOutlinerTreeItem& DropTarget) const;
+	UE_API virtual bool HasCustomDropHandling(const ISceneOutlinerTreeItem& DropTarget) const;
 
 	/** Test whether the specified payload can be dropped onto a tree item */
-	virtual FSceneOutlinerDragValidationInfo ValidateDrop(const ISceneOutlinerTreeItem& DropTarget, const FSceneOutlinerDragDropPayload& Payload) const;
+	UE_API virtual FSceneOutlinerDragValidationInfo ValidateDrop(const ISceneOutlinerTreeItem& DropTarget, const FSceneOutlinerDragDropPayload& Payload) const;
 
 	/** Called when a payload is dropped onto a target */
-	virtual void OnDrop(ISceneOutlinerTreeItem& DropTarget, const FSceneOutlinerDragDropPayload& Payload, const FSceneOutlinerDragValidationInfo& ValidationInfo) const;
+	UE_API virtual void OnDrop(ISceneOutlinerTreeItem& DropTarget, const FSceneOutlinerDragDropPayload& Payload, const FSceneOutlinerDragValidationInfo& ValidationInfo) const;
 
 	/** Called when a context menu is being created from in the hierarchy. Allows subclasses to inject additonal context data. */
 	virtual void OnContextMenuContextCreated(FToolMenuContext& Context) const {}
 
-	static TSet<UClass*> GetParentAndChildClassesFromSpecifiedClasses(
+	static UE_API TSet<UClass*> GetParentAndChildClassesFromSpecifiedClasses(
 		const TSet<UClass*>& InSpecifiedClasses, EObjectMixerInheritanceInclusionOptions Options);
 
-	static TSet<UClass*> GetParentAndChildClassesFromSpecifiedClasses(
+	static UE_API TSet<UClass*> GetParentAndChildClassesFromSpecifiedClasses(
 		const TSet<TSubclassOf<AActor>>& InSpecifiedClasses, EObjectMixerInheritanceInclusionOptions Options);
 };
 
@@ -176,8 +178,8 @@ public:
  * Script class for filtering object types to Object Mixer.
  * Blueprint classes should inherit directly from this class.
  */
-UCLASS(Abstract, Blueprintable)
-class OBJECTMIXEREDITOR_API UObjectMixerBlueprintObjectFilter : public UObjectMixerObjectFilter
+UCLASS(MinimalAPI, Abstract, Blueprintable)
+class UObjectMixerBlueprintObjectFilter : public UObjectMixerObjectFilter
 {
 	GENERATED_BODY()
 public:
@@ -189,7 +191,7 @@ public:
 	 * override the GetObjectMixerPropertyInheritanceInclusionOptions and GetForceAddedColumns functions.
 	 */
 	UFUNCTION(BlueprintNativeEvent, BlueprintCallable, Category = "Object Mixer")
-	TSet<UClass*> GetObjectClassesToFilter() const override;
+	UE_API TSet<UClass*> GetObjectClassesToFilter() const override;
 
 	TSet<UClass*> GetObjectClassesToFilter_Implementation() const
 	{
@@ -202,7 +204,7 @@ public:
 	 * This includes most engine actor types.
 	 */
 	UFUNCTION(BlueprintNativeEvent, BlueprintCallable, Category = "Object Mixer")
-	TSet<TSubclassOf<AActor>> GetObjectClassesToPlace() const override;
+	UE_API TSet<TSubclassOf<AActor>> GetObjectClassesToPlace() const override;
 
 	TSet<TSubclassOf<AActor>> GetObjectClassesToPlace_Implementation() const
 	{
@@ -213,7 +215,7 @@ public:
 	 * Determines if transient objects (such as Sequencer Spawnables) should be shown in the list. False by default.
 	 */
 	UFUNCTION(BlueprintNativeEvent, BlueprintCallable, Category = "Object Mixer")
-	bool GetShowTransientObjects() const;
+	UE_API bool GetShowTransientObjects() const;
 
 	bool GetShowTransientObjects_Implementation() const
 	{
@@ -226,7 +228,7 @@ public:
 	 * Columns not specified will not be shown by default but can be enabled by the user in the UI.
 	 */
 	UFUNCTION(BlueprintNativeEvent, BlueprintCallable, Category = "Object Mixer")
-	TSet<FName> GetColumnsToShowByDefault() const override;
+	UE_API TSet<FName> GetColumnsToShowByDefault() const override;
 	
 	TSet<FName> GetColumnsToShowByDefault_Implementation() const
 	{
@@ -239,7 +241,7 @@ public:
 	 * Columns not specified can be enabled by the user in the UI.
 	 */
 	UFUNCTION(BlueprintNativeEvent, BlueprintCallable, Category = "Object Mixer")
-	TSet<FName> GetColumnsToExclude() const override;
+	UE_API TSet<FName> GetColumnsToExclude() const override;
 
 	TSet<FName> GetColumnsToExclude_Implementation() const
 	{
@@ -256,7 +258,7 @@ public:
 	 * the property is accessible to one of the specified classes regardless of which parent class it comes from.
 	 */
 	UFUNCTION(BlueprintNativeEvent, BlueprintCallable, Category = "Object Mixer")
-	TSet<FName> GetForceAddedColumns() const override;
+	UE_API TSet<FName> GetForceAddedColumns() const override;
 	
 	TSet<FName> GetForceAddedColumns_Implementation() const
 	{
@@ -269,7 +271,7 @@ public:
 	 * If you're not seeing all the properties you expected, try overloading this function.
 	 */
 	UFUNCTION(BlueprintNativeEvent, BlueprintCallable, Category = "Object Mixer")
-	EObjectMixerInheritanceInclusionOptions GetObjectMixerPropertyInheritanceInclusionOptions() const override;
+	UE_API EObjectMixerInheritanceInclusionOptions GetObjectMixerPropertyInheritanceInclusionOptions() const override;
 
 	EObjectMixerInheritanceInclusionOptions GetObjectMixerPropertyInheritanceInclusionOptions_Implementation() const
 	{
@@ -281,7 +283,7 @@ public:
 	 * Defaults to 'None' which only considers the specified classes.
 	 */
 	UFUNCTION(BlueprintNativeEvent, BlueprintCallable, Category = "Object Mixer")
-	EObjectMixerInheritanceInclusionOptions GetObjectMixerPlacementClassInclusionOptions() const override;
+	UE_API EObjectMixerInheritanceInclusionOptions GetObjectMixerPlacementClassInclusionOptions() const override;
 
 	EObjectMixerInheritanceInclusionOptions GetObjectMixerPlacementClassInclusionOptions_Implementation() const
 	{
@@ -293,7 +295,7 @@ public:
 	 * Defaults to false.
 	 */
 	UFUNCTION(BlueprintNativeEvent, BlueprintCallable, Category = "Object Mixer")
-	bool ShouldIncludeUnsupportedProperties() const override;
+	UE_API bool ShouldIncludeUnsupportedProperties() const override;
 
 	bool ShouldIncludeUnsupportedProperties_Implementation() const
 	{
@@ -305,7 +307,7 @@ public:
 	 * Add a property name to this list if you expect the list to change in some way after changing that property.
 	 */
 	UFUNCTION(BlueprintNativeEvent, BlueprintCallable, Category = "Object Mixer")
-	TSet<FName> GetPropertiesThatRequireListRefresh() const override;
+	UE_API TSet<FName> GetPropertiesThatRequireListRefresh() const override;
 
 	TSet<FName> GetPropertiesThatRequireListRefresh_Implementation() const
 	{
@@ -313,3 +315,5 @@ public:
 	}
 
 };
+
+#undef UE_API
diff --git a/Engine/Plugins/Editor/ObjectMixer/ObjectMixer/Source/ObjectMixer/Public/ObjectMixerEditorModule.h b/Engine/Plugins/Editor/ObjectMixer/ObjectMixer/Source/ObjectMixer/Public/ObjectMixerEditorModule.h
--- a/Engine/Plugins/Editor/ObjectMixer/ObjectMixer/Source/ObjectMixer/Public/ObjectMixerEditorModule.h
+++ b/Engine/Plugins/Editor/ObjectMixer/ObjectMixer/Source/ObjectMixer/Public/ObjectMixerEditorModule.h
@@ -6,77 +6,79 @@
 #include "ObjectFilter/ObjectMixerEditorObjectFilter.h"
 #include "Widgets/Docking/SDockTab.h"
 
+#define UE_API OBJECTMIXEREDITOR_API
+
 class FObjectMixerEditorList;
 class ISequencer;
 
-class OBJECTMIXEREDITOR_API FObjectMixerEditorModule : public IModuleInterface
+class FObjectMixerEditorModule : public IModuleInterface
 {
 public:
 
 	//~ Begin IModuleInterface Interface
-	virtual void StartupModule() override;
-	virtual void ShutdownModule() override;
+	UE_API virtual void StartupModule() override;
+	UE_API virtual void ShutdownModule() override;
 	//~ End IModuleInterface
 
-	virtual UWorld* GetWorld();
+	UE_API virtual UWorld* GetWorld();
 
-	virtual void Initialize();
-	virtual void Teardown();
+	UE_API virtual void Initialize();
+	UE_API virtual void Teardown();
 
-	static FObjectMixerEditorModule& Get();
+	static UE_API FObjectMixerEditorModule& Get();
 
-	static void OpenProjectSettings();
+	static UE_API void OpenProjectSettings();
 
-	virtual FName GetModuleName() const;
+	UE_API virtual FName GetModuleName() const;
 	
-	virtual TSharedPtr<SWidget> MakeObjectMixerDialog(
+	UE_API virtual TSharedPtr<SWidget> MakeObjectMixerDialog(
 		TSubclassOf<UObjectMixerObjectFilter> InDefaultFilterClass = nullptr);
 
 	/** Get a list of sequencers that are currently active in the editor. */
-	virtual TArray<TWeakPtr<ISequencer>> GetSequencers() const;
+	UE_API virtual TArray<TWeakPtr<ISequencer>> GetSequencers() const;
 
 	/**
 	 * Tries to find the nomad tab assigned to this instance of Object Mixer.
 	 * If DockTab is not set, will try to find the tab using GetTabSpawnerId().
 	 */
-	virtual TSharedPtr<SDockTab> FindNomadTab();
+	UE_API virtual TSharedPtr<SDockTab> FindNomadTab();
 	
 	/**
 	 * Build the List widget from scratch. If DockTab is not set, will try to find the tab using GetTabSpawnerId().
 	 * @return True if the widget was regenerated. False if the DockTab was invalid and could not be found.
 	 */
-	bool RegenerateListWidget();
+	UE_API bool RegenerateListWidget();
 
 	/** Called when the Rename command is executed from the UI or hotkey. */
-	virtual void OnRenameCommand();
+	UE_API virtual void OnRenameCommand();
 	
-	void RegisterMenuGroup();
-	void UnregisterMenuGroup();
-	virtual void SetupMenuItemVariables();
-	virtual void RegisterTabSpawner();
-	virtual FName GetTabSpawnerId();
+	UE_API void RegisterMenuGroup();
+	UE_API void UnregisterMenuGroup();
+	UE_API virtual void SetupMenuItemVariables();
+	UE_API virtual void RegisterTabSpawner();
+	UE_API virtual FName GetTabSpawnerId();
 	
 	/**
 	 * Add a tab spawner to the Object Mixer menu group.
 	 * @return If adding the item to the menu was successful
 	 */
-	bool RegisterItemInMenuGroup(FWorkspaceItem& InItem);
+	UE_API bool RegisterItemInMenuGroup(FWorkspaceItem& InItem);
 	
-	virtual void UnregisterTabSpawner();
-	virtual void RegisterSettings() const;
-	virtual void UnregisterSettings() const;
+	UE_API virtual void UnregisterTabSpawner();
+	UE_API virtual void RegisterSettings() const;
+	UE_API virtual void UnregisterSettings() const;
 
-	virtual TSharedRef<SDockTab> SpawnTab(const FSpawnTabArgs& Args);
+	UE_API virtual TSharedRef<SDockTab> SpawnTab(const FSpawnTabArgs& Args);
 
-	TSharedPtr<FWorkspaceItem> GetWorkspaceGroup();
+	UE_API TSharedPtr<FWorkspaceItem> GetWorkspaceGroup();
 
 	/**
 	 * This is the filter class used to initialize the ListModel.
 	 * This filter class cannot be turned off by the end user.
 	 */
-	const TSubclassOf<UObjectMixerObjectFilter>& GetDefaultFilterClass() const;
+	UE_API const TSubclassOf<UObjectMixerObjectFilter>& GetDefaultFilterClass() const;
 
-	const static FName BaseObjectMixerModuleName;
+	UE_API const static FName BaseObjectMixerModuleName;
 	
 	DECLARE_MULTICAST_DELEGATE(FOnBlueprintFilterCompiled);
 	FOnBlueprintFilterCompiled& OnBlueprintFilterCompiled()
@@ -113,3 +115,5 @@ private:
 
 	TSharedPtr<FWorkspaceItem> WorkspaceGroup;
 };
+
+#undef UE_API
diff --git a/Engine/Plugins/Editor/ObjectMixer/ObjectMixer/Source/ObjectMixer/Public/ObjectMixerEditorSerializedData.h b/Engine/Plugins/Editor/ObjectMixer/ObjectMixer/Source/ObjectMixer/Public/ObjectMixerEditorSerializedData.h
--- a/Engine/Plugins/Editor/ObjectMixer/ObjectMixer/Source/ObjectMixer/Public/ObjectMixerEditorSerializedData.h
+++ b/Engine/Plugins/Editor/ObjectMixer/ObjectMixer/Source/ObjectMixer/Public/ObjectMixerEditorSerializedData.h
@@ -4,6 +4,8 @@
 
 #include "ObjectMixerEditorSerializedData.generated.h"
 
+#define UE_API OBJECTMIXEREDITOR_API
+
 DECLARE_MULTICAST_DELEGATE(FOnObjectMixerCollectionChanged)
 DECLARE_MULTICAST_DELEGATE_TwoParams(FOnObjectMixerCollectionRenamed, FName, FName)
 
@@ -96,8 +98,8 @@ struct FObjectMixerSerializationDataPerFilter
 	}
 };
 
-UCLASS(config = ObjectMixerSerializedData)
-class OBJECTMIXEREDITOR_API UObjectMixerEditorSerializedData : public UObject
+UCLASS(MinimalAPI, config = ObjectMixerSerializedData)
+class UObjectMixerEditorSerializedData : public UObject
 {
 	GENERATED_BODY()
 public:
@@ -111,36 +113,38 @@ public:
 	UPROPERTY(Config)
 	TSet<FObjectMixerSerializationDataPerFilter> SerializedDataPerFilter;
 
-	FObjectMixerSerializationDataPerFilter* FindSerializationDataByFilterClassName(const FName& FilterClassName);
+	UE_API FObjectMixerSerializationDataPerFilter* FindSerializationDataByFilterClassName(const FName& FilterClassName);
 
-	bool AddObjectsToCollection(const FName& FilterClassName, const FName& CollectionName, const TSet<FSoftObjectPath>& ObjectsToAdd);
+	UE_API bool AddObjectsToCollection(const FName& FilterClassName, const FName& CollectionName, const TSet<FSoftObjectPath>& ObjectsToAdd);
 
-	bool RemoveObjectsFromCollection(const FName& FilterClassName, const FName& CollectionName, const TSet<FSoftObjectPath>& ObjectsToRemove);
+	UE_API bool RemoveObjectsFromCollection(const FName& FilterClassName, const FName& CollectionName, const TSet<FSoftObjectPath>& ObjectsToRemove);
 
-	bool RemoveCollection(const FName& FilterClassName, const FName& CollectionName);
+	UE_API bool RemoveCollection(const FName& FilterClassName, const FName& CollectionName);
 	
-	bool DuplicateCollection(const FName& FilterClassName, const FName& CollectionToDuplicateName, FName& DesiredDuplicateName);
+	UE_API bool DuplicateCollection(const FName& FilterClassName, const FName& CollectionToDuplicateName, FName& DesiredDuplicateName);
 
-	bool ReorderCollection(const FName& FilterClassName, const FName& CollectionToMoveName, const FName& CollectionInsertBeforeName);
+	UE_API bool ReorderCollection(const FName& FilterClassName, const FName& CollectionToMoveName, const FName& CollectionInsertBeforeName);
 
-	bool RenameCollection(const FName& FilterClassName, const FName CollectionNameToRename, const FName NewCollectionName);
+	UE_API bool RenameCollection(const FName& FilterClassName, const FName CollectionNameToRename, const FName NewCollectionName);
 
-	bool DoesCollectionExist(const FName& FilterClassName, const FName& CollectionName);
+	UE_API bool DoesCollectionExist(const FName& FilterClassName, const FName& CollectionName);
 
-	bool IsObjectInCollection(const FName& FilterClassName, const FName& CollectionName, const FSoftObjectPath& InObject);
+	UE_API bool IsObjectInCollection(const FName& FilterClassName, const FName& CollectionName, const FSoftObjectPath& InObject);
 
-	TSet<FName> GetCollectionsForObject(const FName& FilterClassName, const FSoftObjectPath& InObject);
+	UE_API TSet<FName> GetCollectionsForObject(const FName& FilterClassName, const FSoftObjectPath& InObject);
 
-	TArray<FName> GetAllCollectionNames(const FName& FilterClassName);
+	UE_API TArray<FName> GetAllCollectionNames(const FName& FilterClassName);
 
-	void SetShouldShowColumn(const FName& FilterClassName, const FName& ColumnName, const bool bNewShouldShowColumn);
+	UE_API void SetShouldShowColumn(const FName& FilterClassName, const FName& ColumnName, const bool bNewShouldShowColumn);
 
-	bool IsColumnDataSerialized(const FName& FilterClassName, const FName& ColumnName);
+	UE_API bool IsColumnDataSerialized(const FName& FilterClassName, const FName& ColumnName);
 	
-	bool ShouldShowColumn(const FName& FilterClassName, const FName& ColumnName);
+	UE_API bool ShouldShowColumn(const FName& FilterClassName, const FName& ColumnName);
 
-	static FName AllCollectionName;
+	static UE_API FName AllCollectionName;
 
 	FOnObjectMixerCollectionChanged OnObjectMixerCollectionMapChanged;
 	FOnObjectMixerCollectionRenamed OnObjectMixerCollectionRenamed;
 };
+
+#undef UE_API
diff --git a/Engine/Plugins/Editor/ObjectMixer/ObjectMixer/Source/ObjectMixer/Public/ObjectMixerEditorSettings.h b/Engine/Plugins/Editor/ObjectMixer/ObjectMixer/Source/ObjectMixer/Public/ObjectMixerEditorSettings.h
--- a/Engine/Plugins/Editor/ObjectMixer/ObjectMixer/Source/ObjectMixer/Public/ObjectMixerEditorSettings.h
+++ b/Engine/Plugins/Editor/ObjectMixer/ObjectMixer/Source/ObjectMixer/Public/ObjectMixerEditorSettings.h
@@ -14,8 +14,8 @@ enum class EObjectMixerHybridMode : uint8
 	HybridNone UMETA(DisplayName="Do not make hybrid rows", ToolTip="Does not make hybrid rows when an actor with a single matching component is selected. Actor and component rows will be split."),
 };
 
-UCLASS(config = ObjectMixer)
-class OBJECTMIXEREDITOR_API UObjectMixerEditorSettings : public UDeveloperSettings
+UCLASS(MinimalAPI, config = ObjectMixer)
+class UObjectMixerEditorSettings : public UDeveloperSettings
 {
 	GENERATED_BODY()
 	
diff --git a/Engine/Plugins/Editor/ObjectMixer/ObjectMixer/Source/ObjectMixer/Public/ObjectMixerEditorStyle.h b/Engine/Plugins/Editor/ObjectMixer/ObjectMixer/Source/ObjectMixer/Public/ObjectMixerEditorStyle.h
--- a/Engine/Plugins/Editor/ObjectMixer/ObjectMixer/Source/ObjectMixer/Public/ObjectMixerEditorStyle.h
+++ b/Engine/Plugins/Editor/ObjectMixer/ObjectMixer/Source/ObjectMixer/Public/ObjectMixerEditorStyle.h
@@ -4,21 +4,23 @@
 
 #include "Styling/SlateStyle.h"
 
-class OBJECTMIXEREDITOR_API FObjectMixerEditorStyle : FSlateStyleSet
+#define UE_API OBJECTMIXEREDITOR_API
+
+class FObjectMixerEditorStyle : FSlateStyleSet
 {
 public:
 
-	static void Initialize();
+	static UE_API void Initialize();
 
-	static void Shutdown();
+	static UE_API void Shutdown();
 
-	static void ReloadTextures();
+	static UE_API void ReloadTextures();
 
-	static const ISlateStyle& Get();
+	static UE_API const ISlateStyle& Get();
 
-	virtual const FName& GetStyleSetName() const override;
+	UE_API virtual const FName& GetStyleSetName() const override;
 
-	virtual const FSlateBrush* GetBrush(const FName PropertyName, const ANSICHAR* Specifier = nullptr, const ISlateStyle* RequestingStyle = nullptr) const override;
+	UE_API virtual const FSlateBrush* GetBrush(const FName PropertyName, const ANSICHAR* Specifier = nullptr, const ISlateStyle* RequestingStyle = nullptr) const override;
 
 	template< typename WidgetStyleType >
 	static const WidgetStyleType& GetWidgetStyle(FName PropertyName, const ANSICHAR* Specifier = nullptr)
@@ -28,9 +30,11 @@ public:
 
 private:
 
-	static FString GetExternalPluginContent(const FString& PluginName, const FString& RelativePath, const ANSICHAR* Extension);
+	static UE_API FString GetExternalPluginContent(const FString& PluginName, const FString& RelativePath, const ANSICHAR* Extension);
 	
-	static TSharedRef<FSlateStyleSet> Create();
+	static UE_API TSharedRef<FSlateStyleSet> Create();
 
-	static TSharedPtr<FSlateStyleSet> StyleInstance;
+	static UE_API TSharedPtr<FSlateStyleSet> StyleInstance;
 };
+
+#undef UE_API
diff --git a/Engine/Plugins/Editor/ObjectMixer/ObjectMixer/Source/ObjectMixer/Public/ObjectMixerFilterFactory.h b/Engine/Plugins/Editor/ObjectMixer/ObjectMixer/Source/ObjectMixer/Public/ObjectMixerFilterFactory.h
--- a/Engine/Plugins/Editor/ObjectMixer/ObjectMixer/Source/ObjectMixer/Public/ObjectMixerFilterFactory.h
+++ b/Engine/Plugins/Editor/ObjectMixer/ObjectMixer/Source/ObjectMixer/Public/ObjectMixerFilterFactory.h
@@ -6,24 +6,28 @@
 
 #include "ObjectMixerFilterFactory.generated.h"
 
-UCLASS(hidecategories = Object)
-class OBJECTMIXEREDITOR_API UObjectMixerBlueprintFilterFactory : public UFactory
+#define UE_API OBJECTMIXEREDITOR_API
+
+UCLASS(MinimalAPI, hidecategories = Object)
+class UObjectMixerBlueprintFilterFactory : public UFactory
 {
 	GENERATED_BODY()
 public:
-	UObjectMixerBlueprintFilterFactory();
+	UE_API UObjectMixerBlueprintFilterFactory();
 	
 	// The parent class of the created blueprint
     UPROPERTY(VisibleAnywhere, Category="ObjectMixerBlueprintFilterFactory", meta=(AllowAbstract = "", BlueprintBaseOnly = ""))
     TSubclassOf<UObject> ParentClass;
 	
 	//~ Begin UFactory Interface
-	virtual FText GetDisplayName() const override;
-	virtual UObject* FactoryCreateNew(UClass* Class, UObject* InParent, FName Name, EObjectFlags Flags, UObject* Context, FFeedbackContext* Warn) override;
-	virtual FString GetDefaultNewAssetName() const override;
-	virtual uint32 GetMenuCategories() const override;
-	virtual FText GetToolTip() const override;
-	virtual FString GetToolTipDocumentationExcerpt() const override;
+	UE_API virtual FText GetDisplayName() const override;
+	UE_API virtual UObject* FactoryCreateNew(UClass* Class, UObject* InParent, FName Name, EObjectFlags Flags, UObject* Context, FFeedbackContext* Warn) override;
+	UE_API virtual FString GetDefaultNewAssetName() const override;
+	UE_API virtual uint32 GetMenuCategories() const override;
+	UE_API virtual FText GetToolTip() const override;
+	UE_API virtual FString GetToolTipDocumentationExcerpt() const override;
 	//~ Begin UFactory Interface
 	
 };
+
+#undef UE_API
diff --git a/Engine/Plugins/Editor/ObjectMixer/ObjectMixer/Source/ObjectMixer/Public/SelectionInterface/IObjectMixerSelectionInterface.h b/Engine/Plugins/Editor/ObjectMixer/ObjectMixer/Source/ObjectMixer/Public/SelectionInterface/IObjectMixerSelectionInterface.h
--- a/Engine/Plugins/Editor/ObjectMixer/ObjectMixer/Source/ObjectMixer/Public/SelectionInterface/IObjectMixerSelectionInterface.h
+++ b/Engine/Plugins/Editor/ObjectMixer/ObjectMixer/Source/ObjectMixer/Public/SelectionInterface/IObjectMixerSelectionInterface.h
@@ -12,7 +12,7 @@ class UActorComponent;
 /**
  * Provides an interface for synchronizing an ObjectMixer's selection with another part of the editor
  */
-class OBJECTMIXEREDITOR_API IObjectMixerSelectionInterface : public TSharedFromThis<IObjectMixerSelectionInterface>
+class IObjectMixerSelectionInterface : public TSharedFromThis<IObjectMixerSelectionInterface>
 {
 public:
 	DECLARE_MULTICAST_DELEGATE(FOnSelectionChanged);
@@ -44,4 +44,4 @@ public:
 
 	/** Get the event handler that is raised when the synchronized selection changes. */
 	virtual FOnSelectionChanged& OnSelectionChanged() = 0;
-};
\ No newline at end of file
+};
diff --git a/Engine/Plugins/Editor/ObjectMixer/ObjectMixer/Source/ObjectMixer/Public/SelectionInterface/LevelEditorObjectMixerSelectionInterface.h b/Engine/Plugins/Editor/ObjectMixer/ObjectMixer/Source/ObjectMixer/Public/SelectionInterface/LevelEditorObjectMixerSelectionInterface.h
--- a/Engine/Plugins/Editor/ObjectMixer/ObjectMixer/Source/ObjectMixer/Public/SelectionInterface/LevelEditorObjectMixerSelectionInterface.h
+++ b/Engine/Plugins/Editor/ObjectMixer/ObjectMixer/Source/ObjectMixer/Public/SelectionInterface/LevelEditorObjectMixerSelectionInterface.h
@@ -4,26 +4,30 @@
 
 #include "IObjectMixerSelectionInterface.h"
 
+#define UE_API OBJECTMIXEREDITOR_API
+
 /**
  * Provides an interface for the ObjectMixer to synchronize with the level editor's selections (via GEditor)
  */
-class OBJECTMIXEREDITOR_API FLevelEditorObjectMixerSelectionInterface : public IObjectMixerSelectionInterface
+class FLevelEditorObjectMixerSelectionInterface : public IObjectMixerSelectionInterface
 {
 public:
 
-	FLevelEditorObjectMixerSelectionInterface();
-	virtual ~FLevelEditorObjectMixerSelectionInterface() override;
+	UE_API FLevelEditorObjectMixerSelectionInterface();
+	UE_API virtual ~FLevelEditorObjectMixerSelectionInterface() override;
 
 	//~ Begin IObjectMixerSelectionInterface interface
-	virtual void SelectActors(const TArray<AActor*>& InSelectedActors, bool bShouldSelect, bool bSelectEvenIfHidden) override;
-	virtual void SelectComponents(const TArray<UActorComponent*>& InSelectedComponents, bool bShouldSelect, bool bSelectEvenIfHidden) override;
-	virtual TArray<AActor*> GetSelectedActors() const override;
-	virtual TArray<UActorComponent*> GetSelectedComponents() const override;
+	UE_API virtual void SelectActors(const TArray<AActor*>& InSelectedActors, bool bShouldSelect, bool bSelectEvenIfHidden) override;
+	UE_API virtual void SelectComponents(const TArray<UActorComponent*>& InSelectedComponents, bool bShouldSelect, bool bSelectEvenIfHidden) override;
+	UE_API virtual TArray<AActor*> GetSelectedActors() const override;
+	UE_API virtual TArray<UActorComponent*> GetSelectedComponents() const override;
 	virtual FOnSelectionChanged& OnSelectionChanged() override { return SelectionChanged; }
 	//~ End IObjectMixerSelectionInterface interface
 
 private:
-	void OnLevelSelectionChanged(UObject* Obj);
+	UE_API void OnLevelSelectionChanged(UObject* Obj);
 
 	FOnSelectionChanged SelectionChanged;
-};
\ No newline at end of file
+};
+
+#undef UE_API
