Renaming a gameplay tag only lets you rename the specific sub tag you're editing.
Renaming a tag will automatically update sub tags with the new name.
[RN] Renaming a gameplay tag in the Gameplay Tag Manager will automatically rename all sub tags.

diff --git a/Engine/Plugins/Editor/GameplayTagsEditor/Source/GameplayTagsEditor/Private/GameplayTagsEditorModule.cpp b/Engine/Plugins/Editor/GameplayTagsEditor/Source/GameplayTagsEditor/Private/GameplayTagsEditorModule.cpp
--- a/Engine/Plugins/Editor/GameplayTagsEditor/Source/GameplayTagsEditor/Private/GameplayTagsEditorModule.cpp
+++ b/Engine/Plugins/Editor/GameplayTagsEditor/Source/GameplayTagsEditor/Private/GameplayTagsEditorModule.cpp
@@ -849,13 +849,25 @@ public:
 		return bSuccess;
 	}
 
-	virtual bool RenameTagInINI(const FString& TagToRename, const FString& TagToRenameTo) override
+	virtual bool RenameTagInINI(const FString& TagToRename, const FString& TagToRenameTo, bool bRenameChildren) override
 	{
 		FName OldTagName = FName(*TagToRename);
 		FName NewTagName = FName(*TagToRenameTo);
 
 		UGameplayTagsManager& Manager = UGameplayTagsManager::Get();
 
+		if (bRenameChildren)
+		{
+			FGameplayTag ActualTag = Manager.RequestGameplayTag(OldTagName);
+			const FGameplayTagContainer Children = Manager.RequestGameplayTagDirectDescendantsInDictionary(ActualTag, EGameplayTagSelectionType::All);
+
+			for (const FGameplayTag& ChildTag : Children.GetGameplayTagArray())
+			{
+				FString NewTag = ChildTag.ToString().Replace(*(TagToRename + '.'), *(TagToRenameTo + '.'), ESearchCase::CaseSensitive);
+				RenameTagInINI(ChildTag.ToString(), NewTag, true);
+			}
+		}
+
 		FString OldComment, NewComment;
 		FName OldTagSourceName, NewTagSourceName;
 		bool bTagIsExplicit;
diff --git a/Engine/Plugins/Editor/GameplayTagsEditor/Source/GameplayTagsEditor/Private/SRenameGameplayTagDialog.cpp b/Engine/Plugins/Editor/GameplayTagsEditor/Source/GameplayTagsEditor/Private/SRenameGameplayTagDialog.cpp
--- a/Engine/Plugins/Editor/GameplayTagsEditor/Source/GameplayTagsEditor/Private/SRenameGameplayTagDialog.cpp
+++ b/Engine/Plugins/Editor/GameplayTagsEditor/Source/GameplayTagsEditor/Private/SRenameGameplayTagDialog.cpp
@@ -67,7 +67,7 @@ void SRenameGameplayTagDialog::Construct(const FArguments& InArgs)
 			[
 				SAssignNew(NewTagNameTextBox, SEditableTextBox)
 				.Font(FAppStyle::GetFontStyle( TEXT("PropertyWindow.NormalFont")))
-				.Text(FText::FromName(GameplayTagNode->GetCompleteTag().GetTagName()))
+				.Text(FText::FromName(GameplayTagNode->GetSimpleTagName()))
 				.OnTextCommitted(this, &SRenameGameplayTagDialog::OnRenameTextCommitted)
 			]
 
@@ -128,9 +128,9 @@ void SRenameGameplayTagDialog::RenameAndClose()
 	IGameplayTagsEditorModule& Module = IGameplayTagsEditorModule::Get();
 
 	const FString TagToRename = GameplayTagNode->GetCompleteTag().GetTagName().ToString();
-	const FString NewTagName = NewTagNameTextBox->GetText().ToString();
+	const FString NewTagName = GameplayTagNode->GetParentTagNode().IsValid() ? GameplayTagNode->GetParentTagNode()->GetCompleteTag().GetTagName().ToString() + '.' + NewTagNameTextBox->GetText().ToString() : NewTagNameTextBox->GetText().ToString();
 
-	if (Module.RenameTagInINI(TagToRename, NewTagName))
+	if (Module.RenameTagInINI(TagToRename, NewTagName, true))
 	{
 		OnGameplayTagRenamed.ExecuteIfBound(TagToRename, NewTagName);
 	}
diff --git a/Engine/Plugins/Editor/GameplayTagsEditor/Source/GameplayTagsEditor/Public/GameplayTagsEditorModule.h b/Engine/Plugins/Editor/GameplayTagsEditor/Source/GameplayTagsEditor/Public/GameplayTagsEditorModule.h
--- a/Engine/Plugins/Editor/GameplayTagsEditor/Source/GameplayTagsEditor/Public/GameplayTagsEditorModule.h
+++ b/Engine/Plugins/Editor/GameplayTagsEditor/Source/GameplayTagsEditor/Public/GameplayTagsEditorModule.h
@@ -54,7 +54,7 @@ public:
 	virtual void DeleteTagsFromINI(const TArray<TSharedPtr<struct FGameplayTagNode>>& TagNodesToDelete) = 0;
 
 	/** Tries to rename a tag, leaving a rediretor in the ini, and adding the new tag if it does not exist yet */
-	virtual bool RenameTagInINI(const FString& TagToRename, const FString& TagToRenameTo) = 0;
+	virtual bool RenameTagInINI(const FString& TagToRename, const FString& TagToRenameTo, bool bRenameChildren = true) = 0;
 
 	/** Updates info about a tag */
 	virtual bool UpdateTagInINI(const FString& TagToUpdate, const FString& Comment, bool bIsRestrictedTag, bool bAllowNonRestrictedChildren) = 0;
