Allow filtering of the gameplay tag chip display using DisplayShortTagNames metadata.
[RN] Gameplay Tag and Tag Container meta data "DisplayShortTagNames" will remove the Categories filter from the beginning of the Tag chip in the details panel.

diff --git a/Engine/Plugins/Editor/GameplayTagsEditor/Source/GameplayTagsEditor/Private/SGameplayTagCombo.cpp b/Engine/Plugins/Editor/GameplayTagsEditor/Source/GameplayTagsEditor/Private/SGameplayTagCombo.cpp
--- a/Engine/Plugins/Editor/GameplayTagsEditor/Source/GameplayTagsEditor/Private/SGameplayTagCombo.cpp
+++ b/Engine/Plugins/Editor/GameplayTagsEditor/Source/GameplayTagsEditor/Private/SGameplayTagCombo.cpp
@@ -35,6 +35,7 @@ void SGameplayTagCombo::Construct(const FArguments& InArgs)
 	Filter = InArgs._Filter;
 	SettingsName = InArgs._SettingsName;
 	bIsReadOnly = InArgs._ReadOnly;
+	bDisplayShortTagNames = InArgs._DisplayShortTagNames;
 	OnTagChanged = InArgs._OnTagChanged;
 	PropertyHandle = InArgs._PropertyHandle;
 
@@ -48,6 +49,7 @@ void SGameplayTagCombo::Construct(const FArguments& InArgs)
 			Filter = UGameplayTagsManager::Get().GetCategoriesMetaFromPropertyHandle(PropertyHandle);
 		}
 		bIsReadOnly = PropertyHandle->IsEditConst();
+		bDisplayShortTagNames |= PropertyHandle->GetBoolMetaData("DisplayShortTagNames");
 	}
 	
 	ChildSlot
@@ -143,6 +145,18 @@ FText SGameplayTagCombo::GetText() const
 			return FText::FromName(GameplayTag.GetTagLeafName());
 		}
 
+		if (bDisplayShortTagNames)
+		{
+			FString TagName = GameplayTag.GetTagName().ToString();
+
+			if (TagName.Len() > Filter.Len())
+			{
+				TagName.RemoveFromStart(Filter + TEXT("."));
+			}
+
+			return FText::FromString(MoveTemp(TagName));
+		}
+
 		return FText::FromName(GameplayTag.GetTagName());
 	}
 	return FText::FromName(TagAttribute.Get().GetTagName());
diff --git a/Engine/Plugins/Editor/GameplayTagsEditor/Source/GameplayTagsEditor/Private/SGameplayTagContainerCombo.cpp b/Engine/Plugins/Editor/GameplayTagsEditor/Source/GameplayTagsEditor/Private/SGameplayTagContainerCombo.cpp
--- a/Engine/Plugins/Editor/GameplayTagsEditor/Source/GameplayTagsEditor/Private/SGameplayTagContainerCombo.cpp
+++ b/Engine/Plugins/Editor/GameplayTagsEditor/Source/GameplayTagsEditor/Private/SGameplayTagContainerCombo.cpp
@@ -40,6 +40,7 @@ void SGameplayTagContainerCombo::Construct(const FArguments& InArgs)
 	Filter = InArgs._Filter;
 	SettingsName = InArgs._SettingsName;
 	bIsReadOnly = InArgs._ReadOnly;
+	bDisplayShortTagNames = InArgs._DisplayShortTagNames;
 	OnTagCleared = InArgs._OnTagCleared;
 	OnTagContainerComboClosed = InArgs._OnTagContainerComboClosed;
 	OnTagContainerChanged = InArgs._OnTagContainerChanged;
@@ -55,6 +56,7 @@ void SGameplayTagContainerCombo::Construct(const FArguments& InArgs)
 			Filter = UGameplayTagsManager::Get().GetCategoriesMetaFromPropertyHandle(PropertyHandle);
 		}
 		bIsReadOnly = PropertyHandle->IsEditConst();
+		bDisplayShortTagNames |= PropertyHandle->GetBoolMetaData("DisplayShortTagNames");
 	}
 	else
 	{
@@ -170,6 +172,15 @@ bool SGameplayTagContainerCombo::IsValueEnabled() const
 
 TSharedRef<ITableRow> SGameplayTagContainerCombo::MakeTagListViewRow(TSharedPtr<FEditableItem> Item, const TSharedRef<STableViewBase>& OwnerTable)
 {
+	FString TagName = Item->Tag.GetTagName().ToString();
+	if (bDisplayShortTagNames)
+	{
+		if (TagName.Len() > Filter.Len())
+		{
+			TagName.RemoveFromStart(Filter + TEXT("."));
+		}
+	}
+	
 	return SNew(STableRow<TSharedPtr<FEditableItem>>, OwnerTable)
 		.Style(&FAppStyle::Get().GetWidgetStyle<FTableRowStyle>("SimpleTableView.Row"))
 		.Padding(FMargin(0,2))
@@ -177,7 +188,7 @@ TSharedRef<ITableRow> SGameplayTagContainerCombo::MakeTagListViewRow(TSharedPtr<
 			SNew(SGameplayTagChip)
 			.ReadOnly(bIsReadOnly)
 			.ShowClearButton(true)
-			.Text(FText::FromName(Item->Tag.GetTagName()))
+			.Text(FText::FromString(TagName))
 			.ToolTipText(this, &SGameplayTagContainerCombo::GetTagToolTipText, Item->Tag)
 			.IsSelected(!Item->bMultipleValues)
 			.OnClearPressed(this, &SGameplayTagContainerCombo::OnClearTagClicked, Item->Tag)
diff --git a/Engine/Plugins/Editor/GameplayTagsEditor/Source/GameplayTagsEditor/Public/SGameplayTagCombo.h b/Engine/Plugins/Editor/GameplayTagsEditor/Source/GameplayTagsEditor/Public/SGameplayTagCombo.h
--- a/Engine/Plugins/Editor/GameplayTagsEditor/Source/GameplayTagsEditor/Public/SGameplayTagCombo.h
+++ b/Engine/Plugins/Editor/GameplayTagsEditor/Source/GameplayTagsEditor/Public/SGameplayTagCombo.h
@@ -29,6 +29,7 @@ public:
 	SLATE_BEGIN_ARGS(SGameplayTagCombo)
 		: _Filter()
 		, _ReadOnly(false)
+		, _DisplayShortTagNames(false)
 		, _EnableNavigation(false)
 		, _PropertyHandle(nullptr)
 	{}
@@ -41,6 +42,9 @@ public:
 		// Flag to set if the list is read only
 		SLATE_ARGUMENT(bool, ReadOnly)
 
+		// Hide Filter from chip display
+		SLATE_ARGUMENT(bool, DisplayShortTagNames)
+
 		// If true, allow button navigation behavior
 		SLATE_ARGUMENT(bool, EnableNavigation)
 
@@ -90,6 +94,7 @@ private:
 	TArray<FGameplayTag> TagsFromProperty;
 	bool bHasMultipleValues = false;
 	bool bIsReadOnly = false;
+	bool bDisplayShortTagNames = false;
 	FString Filter;
 	FString SettingsName;
 	FOnTagChanged OnTagChanged;
diff --git a/Engine/Plugins/Editor/GameplayTagsEditor/Source/GameplayTagsEditor/Public/SGameplayTagContainerCombo.h b/Engine/Plugins/Editor/GameplayTagsEditor/Source/GameplayTagsEditor/Public/SGameplayTagContainerCombo.h
--- a/Engine/Plugins/Editor/GameplayTagsEditor/Source/GameplayTagsEditor/Public/SGameplayTagContainerCombo.h
+++ b/Engine/Plugins/Editor/GameplayTagsEditor/Source/GameplayTagsEditor/Public/SGameplayTagContainerCombo.h
@@ -33,6 +33,7 @@ public:
 	SLATE_BEGIN_ARGS(SGameplayTagContainerCombo)
 		: _Filter()
 		, _ReadOnly(false)
+		, _DisplayShortTagNames(false)
 		, _EnableNavigation(false)
 		, _PropertyHandle(nullptr)
 	{}
@@ -45,6 +46,9 @@ public:
 		// Flag to set if the list is read only
 		SLATE_ARGUMENT(bool, ReadOnly)
 
+		// Hide Filter from chip display
+		SLATE_ARGUMENT(bool, DisplayShortTagNames)
+
 		// If true, allow button navigation behavior
 		SLATE_ARGUMENT(bool, EnableNavigation)
 
@@ -119,6 +123,7 @@ private:
 	FString Filter;
 	FString SettingsName;
 	bool bIsReadOnly = false;
+	bool bDisplayShortTagNames = false;
 	FOnTagCleared OnTagCleared;
 	FOnTagContainerComboClosed OnTagContainerComboClosed;
 	FOnTagContainerChanged OnTagContainerChanged;
