Tag chip ends in ellipsis when cut off.
[RN] Gameplay Tag display ends in ellipsis when the display is cut off rather than overflowing.

diff --git a/Engine/Plugins/Editor/GameplayTagsEditor/Source/GameplayTagsEditor/Private/SGameplayTagChip.cpp b/Engine/Plugins/Editor/GameplayTagsEditor/Source/GameplayTagsEditor/Private/SGameplayTagChip.cpp
--- a/Engine/Plugins/Editor/GameplayTagsEditor/Source/GameplayTagsEditor/Private/SGameplayTagChip.cpp
+++ b/Engine/Plugins/Editor/GameplayTagsEditor/Source/GameplayTagsEditor/Private/SGameplayTagChip.cpp
@@ -117,6 +117,7 @@ void SGameplayTagChip::Construct(const FArguments& InArgs)
 						[
 							SNew(STextBlock)
 							.Font(FAppStyle::GetFontStyle( TEXT("PropertyWindow.NormalFont")))
+							.OverflowPolicy(ETextOverflowPolicy::Ellipsis)
 							.Text_Lambda([WeakSelf]()
 							{
 								const TSharedPtr<SGameplayTagChip> Self = WeakSelf.Pin();
