GameplayTags: Tokenize search string so it's possible to find "Best.Tag.Ever" by looking for "best eve" and other whitespace-delimited combinations.

#rnx
#Synced-CL 48848898

diff --git a/Engine/Plugins/Editor/GameplayTagsEditor/Source/GameplayTagsEditor/Private/SGameplayTagPicker.cpp b/Engine/Plugins/Editor/GameplayTagsEditor/Source/GameplayTagsEditor/Private/SGameplayTagPicker.cpp
--- a/Engine/Plugins/Editor/GameplayTagsEditor/Source/GameplayTagsEditor/Private/SGameplayTagPicker.cpp
+++ b/Engine/Plugins/Editor/GameplayTagsEditor/Source/GameplayTagsEditor/Private/SGameplayTagPicker.cpp
@@ -233,7 +233,7 @@ void SGameplayTagPicker::Construct(const FArguments& InArgs)
 				.Padding(0,1,5,1)
 				[
 					SAssignNew(SearchTagBox, SSearchBox)
-					.HintText(LOCTEXT("GameplayTagPicker_SearchBoxHint", "Search Gameplay Tags"))
+					.HintText(LOCTEXT("GameplayTagPicker_SearchBoxHint", "Search tags matching ALL keywords"))
 					.OnTextChanged(this, &SGameplayTagPicker::OnFilterTextChanged)
 				]
 
@@ -474,13 +474,13 @@ void SGameplayTagPicker::Tick(const FGeometry& AllottedGeometry, const double In
 
 void SGameplayTagPicker::OnFilterTextChanged(const FText& InFilterText)
 {
-	FilterString = InFilterText.ToString();	
+	InFilterText.ToString().ParseIntoArrayWS(FilterStringTokens);
 	FilterTagTree();
 }
 
 void SGameplayTagPicker::FilterTagTree()
 {
-	if (FilterString.IsEmpty())
+	if (FilterStringTokens.IsEmpty())
 	{
 		TagTreeWidget->SetTreeItemsSource(&TagItems);
 
@@ -550,17 +550,25 @@ bool SGameplayTagPicker::FilterChildrenCheck(TSharedPtr<FGameplayTagNode>& InIte
 		return FilterChildrenCheckRecursive(InItem);
 	}
 
-	if (InItem->GetCompleteTagString().Contains(FilterString) || FilterString.IsEmpty())
+	if (!FilterStringTokens.IsEmpty())
 	{
-		return true;
+		const FString CompleteTagString = InItem->GetCompleteTagString();
+		for (const FString& Token : FilterStringTokens)
+		{
+			if (!CompleteTagString.Contains(Token))
+			{
+				return FilterChildrenCheckRecursive(InItem);
+			}
+		}
 	}
 
-	return FilterChildrenCheckRecursive(InItem);
+	// All tokens matched (or there are no search tokens)
+	return true;
 }
 
 FText SGameplayTagPicker::GetHighlightText() const
 {
-	return FilterString.IsEmpty() ? FText::GetEmpty() : FText::FromString(FilterString);
+	return FilterStringTokens.IsEmpty() ? FText::GetEmpty() : FText::FromString(FilterStringTokens.Last());
 }
 
 TSharedRef<ITableRow> SGameplayTagPicker::OnGenerateRow(TSharedPtr<FGameplayTagNode> InItem, const TSharedRef<STableViewBase>& OwnerTable)
diff --git a/Engine/Plugins/Editor/GameplayTagsEditor/Source/GameplayTagsEditor/Public/SGameplayTagPicker.h b/Engine/Plugins/Editor/GameplayTagsEditor/Source/GameplayTagsEditor/Public/SGameplayTagPicker.h
--- a/Engine/Plugins/Editor/GameplayTagsEditor/Source/GameplayTagsEditor/Public/SGameplayTagPicker.h
+++ b/Engine/Plugins/Editor/GameplayTagsEditor/Source/GameplayTagsEditor/Public/SGameplayTagPicker.h
@@ -385,8 +385,8 @@ private:
 	/* Holds the Name of this TagContainer used for saving out expansion settings */
 	FString SettingsName;
 
-	/* Filter string used during search box */
-	FString FilterString;
+	/* Filter string tokens used during search box */
+	TArray<FString> FilterStringTokens;
 
 	/** root filter (passed in on creation) */
 	FString RootFilterString;
