Fix for Gameplay Tag /Container customization using the HasMetaData("Categories") instead of the proper category metadata getter that includes metadata delegates.

#Synced-CL 49000457

diff --git a/Engine/Plugins/Editor/GameplayTagsEditor/Source/GameplayTagsEditor/Private/GameplayTagContainerCustomization.cpp b/Engine/Plugins/Editor/GameplayTagsEditor/Source/GameplayTagsEditor/Private/GameplayTagContainerCustomization.cpp
--- a/Engine/Plugins/Editor/GameplayTagsEditor/Source/GameplayTagsEditor/Private/GameplayTagContainerCustomization.cpp
+++ b/Engine/Plugins/Editor/GameplayTagsEditor/Source/GameplayTagsEditor/Private/GameplayTagContainerCustomization.cpp
@@ -37,10 +37,10 @@ void FGameplayTagContainerCustomization::CustomizeHeader(TSharedRef<IPropertyHan
 	FString CurrentValue;
 	StructPropertyHandle->GetValueAsFormattedString(CurrentValue);
 
-	if (StructPropertyHandle->HasMetaData("Categories"))
+	const FString CategoriesString = UGameplayTagsManager::Get().GetCategoriesMetaFromPropertyHandle(StructPropertyHandle);
+	if (!CategoriesString.IsEmpty())
 	{
 		TArray<FString> CategoriesStringList;
-		const FString CategoriesString = UGameplayTagsManager::Get().GetCategoriesMetaFromPropertyHandle(StructPropertyHandle);
 
 		if (CategoriesString.ParseIntoArray(CategoriesStringList, TEXT(","), true) > 0)
 		{
diff --git a/Engine/Plugins/Editor/GameplayTagsEditor/Source/GameplayTagsEditor/Private/GameplayTagCustomization.cpp b/Engine/Plugins/Editor/GameplayTagsEditor/Source/GameplayTagsEditor/Private/GameplayTagCustomization.cpp
--- a/Engine/Plugins/Editor/GameplayTagsEditor/Source/GameplayTagsEditor/Private/GameplayTagCustomization.cpp
+++ b/Engine/Plugins/Editor/GameplayTagsEditor/Source/GameplayTagsEditor/Private/GameplayTagCustomization.cpp
@@ -44,10 +44,10 @@ void FGameplayTagCustomization::CustomizeHeader(TSharedRef<IPropertyHandle> InSt
 	FString CurrentValue;
 	StructPropertyHandle->GetValueAsFormattedString(CurrentValue);
 
-	if (StructPropertyHandle->HasMetaData("Categories"))
+	const FString CategoriesString = UGameplayTagsManager::Get().GetCategoriesMetaFromPropertyHandle(StructPropertyHandle);
+	if (!CategoriesString.IsEmpty())
 	{
 		TArray<FString> CategoriesStringList;
-		const FString CategoriesString = UGameplayTagsManager::Get().GetCategoriesMetaFromPropertyHandle(StructPropertyHandle);
 
 		if (CategoriesString.ParseIntoArray(CategoriesStringList, TEXT(","), true) > 0)
 		{
@@ -105,10 +105,10 @@ void FGameplayTagCustomization::CustomizeChildren(TSharedRef<IPropertyHandle> In
 		}
 	}
 
-	if (StructPropertyHandle->HasMetaData("Categories"))
+	const FString CategoriesString = UGameplayTagsManager::Get().GetCategoriesMetaFromPropertyHandle(StructPropertyHandle);
+	if (!CategoriesString.IsEmpty())
 	{
 		TArray<FString> CategoriesStringList;
-		const FString CategoriesString = UGameplayTagsManager::Get().GetCategoriesMetaFromPropertyHandle(StructPropertyHandle);
 
 		if (CategoriesString.ParseIntoArray(CategoriesStringList, TEXT(","), true) > 0)
 		{
