Add accessors to the gameplay tag container UI so it can be opened by other UI elements
#Synced-CL 49736630

diff --git a/Engine/Plugins/Editor/GameplayTagsEditor/Source/GameplayTagsEditor/Private/SGameplayTagContainerCombo.cpp b/Engine/Plugins/Editor/GameplayTagsEditor/Source/GameplayTagsEditor/Private/SGameplayTagContainerCombo.cpp
--- a/Engine/Plugins/Editor/GameplayTagsEditor/Source/GameplayTagsEditor/Private/SGameplayTagContainerCombo.cpp
+++ b/Engine/Plugins/Editor/GameplayTagsEditor/Source/GameplayTagsEditor/Private/SGameplayTagContainerCombo.cpp
@@ -361,14 +361,23 @@ void SGameplayTagContainerCombo::OnSearchForAnyReferences() const
 	}
 }
 
+bool SGameplayTagContainerCombo::IsOpen() const
+{
+	return ComboButton.IsValid() && ComboButton->IsOpen();
+}
+
+void SGameplayTagContainerCombo::Open(const FGameplayTag TagToFocus)
+{
+	TagToHilight = TagToFocus;
+	ComboButton->SetIsOpen(true);
+}
+
 FReply SGameplayTagContainerCombo::OnEditClicked(const FGameplayTag InTagToHilight)
 {
 	FReply Reply = FReply::Handled();
 	if (ComboButton->ShouldOpenDueToClick())
 	{
-		TagToHilight = InTagToHilight;
-		
-		ComboButton->SetIsOpen(true);
+		Open(InTagToHilight);
 		
 		if (TagPicker.IsValid() && TagPicker->GetWidgetToFocusOnOpen())
 		{
diff --git a/Engine/Plugins/Editor/GameplayTagsEditor/Source/GameplayTagsEditor/Public/SGameplayTagContainerCombo.h b/Engine/Plugins/Editor/GameplayTagsEditor/Source/GameplayTagsEditor/Public/SGameplayTagContainerCombo.h
--- a/Engine/Plugins/Editor/GameplayTagsEditor/Source/GameplayTagsEditor/Public/SGameplayTagContainerCombo.h
+++ b/Engine/Plugins/Editor/GameplayTagsEditor/Source/GameplayTagsEditor/Public/SGameplayTagContainerCombo.h
@@ -80,6 +80,12 @@ public:
 
 	GAMEPLAYTAGSEDITOR_API virtual void Tick(const FGeometry& AllottedGeometry, const double InCurrentTime, const float InDeltaTime) override;
 
+	/** Returns true if the expanded edit UI is open */
+	GAMEPLAYTAGSEDITOR_API bool IsOpen() const;
+
+	/** Opens the expanded edit UI, optionally focusing on a specific tag */
+	GAMEPLAYTAGSEDITOR_API void Open(FGameplayTag TagToFocus = { });
+
 private:
 
 	struct FEditableItem
