Implementing an option to color gameplaytag chips within the gameplaytageditor for feedback.

#Synced-CL 49955713

#Synced-CL 49955713

#Synced-CL 49225623

diff --git a/Engine/Plugins/Editor/GameplayTagsEditor/Source/GameplayTagsEditor/Private/SGameplayTagChip.cpp b/Engine/Plugins/Editor/GameplayTagsEditor/Source/GameplayTagsEditor/Private/SGameplayTagChip.cpp
--- a/Engine/Plugins/Editor/GameplayTagsEditor/Source/GameplayTagsEditor/Private/SGameplayTagChip.cpp
+++ b/Engine/Plugins/Editor/GameplayTagsEditor/Source/GameplayTagsEditor/Private/SGameplayTagChip.cpp
@@ -45,6 +45,7 @@ void SGameplayTagChip::Construct(const FArguments& InArgs)
 	ShowClearButtonAttribute.Assign(*this, InArgs._ShowClearButton);
 	ToolTipTextAttribute.Assign(*this, InArgs._ToolTipText);
 	TextAttribute.Assign(*this, InArgs._Text);
+	TextColor = InArgs._TextColor;
 	OnClearPressed = InArgs._OnClearPressed;
 	OnEditPressed = InArgs._OnEditPressed;
 	OnNavigate = InArgs._OnNavigate;
@@ -115,7 +116,7 @@ void SGameplayTagChip::Construct(const FArguments& InArgs)
 						.AutoWidth()
 						.VAlign(VAlign_Center)
 						[
-							SNew(STextBlock)
+							SAssignNew(TagText, STextBlock)
 							.Font(FAppStyle::GetFontStyle( TEXT("PropertyWindow.NormalFont")))
 							.OverflowPolicy(ETextOverflowPolicy::Ellipsis)
 							.Text_Lambda([WeakSelf]()
@@ -202,6 +203,11 @@ void SGameplayTagChip::UpdatePillStyle()
 		}
 		bLastHasIsSelected = bIsSelected;
 	}
+	
+	if (TextColor.IsSet())
+	{
+		TagText->SetColorAndOpacity(TextColor.GetValue());
+	}
 }
 
 #undef LOCTEXT_NAMESPACE
diff --git a/Engine/Plugins/Editor/GameplayTagsEditor/Source/GameplayTagsEditor/Private/SGameplayTagContainerCombo.cpp b/Engine/Plugins/Editor/GameplayTagsEditor/Source/GameplayTagsEditor/Private/SGameplayTagContainerCombo.cpp
--- a/Engine/Plugins/Editor/GameplayTagsEditor/Source/GameplayTagsEditor/Private/SGameplayTagContainerCombo.cpp
+++ b/Engine/Plugins/Editor/GameplayTagsEditor/Source/GameplayTagsEditor/Private/SGameplayTagContainerCombo.cpp
@@ -44,6 +44,7 @@ void SGameplayTagContainerCombo::Construct(const FArguments& InArgs)
 	OnTagCleared = InArgs._OnTagCleared;
 	OnTagContainerComboClosed = InArgs._OnTagContainerComboClosed;
 	OnTagContainerChanged = InArgs._OnTagContainerChanged;
+	OnGetTagTextColor = InArgs._OnGetTagTextColor;
 	PropertyHandle = InArgs._PropertyHandle;
 
 	if (PropertyHandle.IsValid())
@@ -181,6 +182,12 @@ TSharedRef<ITableRow> SGameplayTagContainerCombo::MakeTagListViewRow(TSharedPtr<
 		}
 	}
 	
+	TOptional<FSlateColor> TextColor;
+	if (OnGetTagTextColor.IsBound())
+	{
+		TextColor = OnGetTagTextColor.Execute(Item->Tag);
+	}
+
 	return SNew(STableRow<TSharedPtr<FEditableItem>>, OwnerTable)
 		.Style(&FAppStyle::Get().GetWidgetStyle<FTableRowStyle>("SimpleTableView.Row"))
 		.Padding(FMargin(0,2))
@@ -189,6 +196,7 @@ TSharedRef<ITableRow> SGameplayTagContainerCombo::MakeTagListViewRow(TSharedPtr<
 			.ReadOnly(bIsReadOnly)
 			.ShowClearButton(true)
 			.Text(FText::FromString(TagName))
+			.TextColor(TextColor)
 			.ToolTipText(this, &SGameplayTagContainerCombo::GetTagToolTipText, Item->Tag)
 			.IsSelected(!Item->bMultipleValues)
 			.OnClearPressed(this, &SGameplayTagContainerCombo::OnClearTagClicked, Item->Tag)
diff --git a/Engine/Plugins/Editor/GameplayTagsEditor/Source/GameplayTagsEditor/Public/SGameplayTagChip.h b/Engine/Plugins/Editor/GameplayTagsEditor/Source/GameplayTagsEditor/Public/SGameplayTagChip.h
--- a/Engine/Plugins/Editor/GameplayTagsEditor/Source/GameplayTagsEditor/Public/SGameplayTagChip.h
+++ b/Engine/Plugins/Editor/GameplayTagsEditor/Source/GameplayTagsEditor/Public/SGameplayTagChip.h
@@ -9,6 +9,7 @@
 #include "GameplayTagContainer.h"
 
 class SButton;
+class STextBlock;
 
 /**
  * Widget for displaying a single gameplay tag with optional clear button.
@@ -59,6 +60,9 @@ public:
 
 		// Text to display
 		SLATE_ATTRIBUTE(FText, Text)
+
+		// Text to display color
+		SLATE_ARGUMENT(TOptional<FSlateColor>, TextColor)
 	SLATE_END_ARGS();
 
 	GAMEPLAYTAGSEDITOR_API SGameplayTagChip();
@@ -73,8 +77,10 @@ private:
 	TSlateAttribute<bool> ShowClearButtonAttribute;
 	TSlateAttribute<FText> ToolTipTextAttribute;
 	TSlateAttribute<FText> TextAttribute;
+	TOptional<FSlateColor> TextColor;
 	TSharedPtr<SButton> ChipButton;
 	TSharedPtr<SButton> ClearButton;
+	TSharedPtr<STextBlock> TagText;
 	FOnClearPressed OnClearPressed;
 	FOnEditPressed OnEditPressed;
 	FOnNavigate OnNavigate;
diff --git a/Engine/Plugins/Editor/GameplayTagsEditor/Source/GameplayTagsEditor/Public/SGameplayTagContainerCombo.h b/Engine/Plugins/Editor/GameplayTagsEditor/Source/GameplayTagsEditor/Public/SGameplayTagContainerCombo.h
--- a/Engine/Plugins/Editor/GameplayTagsEditor/Source/GameplayTagsEditor/Public/SGameplayTagContainerCombo.h
+++ b/Engine/Plugins/Editor/GameplayTagsEditor/Source/GameplayTagsEditor/Public/SGameplayTagContainerCombo.h
@@ -29,6 +29,8 @@ public:
 	DECLARE_DELEGATE(FOnTagCleared)
 	DECLARE_DELEGATE_OneParam(FOnTagContainerComboClosed, const FGameplayTagContainer& /*TagContainer*/)
 	DECLARE_DELEGATE_OneParam(FOnTagContainerChanged, const FGameplayTagContainer& /*TagContainer*/)
+	// If bound returns optional color to use for this tag text. 
+    DECLARE_DELEGATE_RetVal_OneParam(TOptional<FSlateColor>, FOnGetTagTextColor, const FGameplayTag& /*Tag*/);
 
 	SLATE_BEGIN_ARGS(SGameplayTagContainerCombo)
 		: _Filter()
@@ -72,6 +74,9 @@ public:
 
 		// Called when a tag container changes
 		SLATE_EVENT(FOnTagContainerChanged, OnTagContainerChanged)
+
+		// Called to request option tag text color
+		SLATE_EVENT(FOnGetTagTextColor, OnGetTagTextColor)
 	SLATE_END_ARGS();
 
 	GAMEPLAYTAGSEDITOR_API SGameplayTagContainerCombo();
@@ -85,6 +90,8 @@ public:
 
 	/** Opens the expanded edit UI, optionally focusing on a specific tag */
 	GAMEPLAYTAGSEDITOR_API void Open(FGameplayTag TagToFocus = { });
+	
+	GAMEPLAYTAGSEDITOR_API void RefreshTagContainers();
 
 private:
 
@@ -103,7 +110,6 @@ private:
 	};
 
 	bool IsValueEnabled() const;
-	void RefreshTagContainers();
 	TSharedRef<ITableRow> MakeTagListViewRow(TSharedPtr<FEditableItem> Item, const TSharedRef<STableViewBase>& OwnerTable);
 	void OnMenuOpenChanged(const bool bOpen);
 	TSharedRef<SWidget> OnGetMenuContent();
@@ -133,6 +139,7 @@ private:
 	FOnTagCleared OnTagCleared;
 	FOnTagContainerComboClosed OnTagContainerComboClosed;
 	FOnTagContainerChanged OnTagContainerChanged;
+	FOnGetTagTextColor OnGetTagTextColor;
 	TSharedPtr<IPropertyHandle> PropertyHandle;
 	TSharedPtr<SMenuAnchor> MenuAnchor;
 	TSharedPtr<SComboButton> ComboButton;
