[UE] Make gameplay tag text searchable in Details panel search bars

Adds .FilterString to FGameplayTag and FGameplayTagContainer property customizations so the tag string value is included in the editor search index. Searching a Details panel for a tag name like "Ability.Jump" now surfaces the property row.

[RN] Gameplay tags and gameplay tag containers are not filterable in details panel by text contained in the tag.
#Synced-CL 52099948

diff --git a/Engine/Plugins/Editor/GameplayTagsEditor/Source/GameplayTagsEditor/Private/GameplayTagContainerCustomization.cpp b/Engine/Plugins/Editor/GameplayTagsEditor/Source/GameplayTagsEditor/Private/GameplayTagContainerCustomization.cpp
--- a/Engine/Plugins/Editor/GameplayTagsEditor/Source/GameplayTagsEditor/Private/GameplayTagContainerCustomization.cpp
+++ b/Engine/Plugins/Editor/GameplayTagsEditor/Source/GameplayTagsEditor/Private/GameplayTagContainerCustomization.cpp
@@ -47,6 +47,7 @@ void FGameplayTagContainerCustomization::CustomizeHeader(TSharedRef<IPropertyHan
 	bool bHasValidCategories = CurrentTagContainer.IsEmpty() || MetaDataCategories.IsEmpty() || CurrentTagContainer.HasAll(MetaDataCategories);
 
 	HeaderRow
+		.FilterString(FText::AsCultureInvariant(CurrentTagContainer.ToStringSimple()))
 		.NameContent()
 		[
 			SNew(SHorizontalBox)
diff --git a/Engine/Plugins/Editor/GameplayTagsEditor/Source/GameplayTagsEditor/Private/GameplayTagCustomization.cpp b/Engine/Plugins/Editor/GameplayTagsEditor/Source/GameplayTagsEditor/Private/GameplayTagCustomization.cpp
--- a/Engine/Plugins/Editor/GameplayTagsEditor/Source/GameplayTagsEditor/Private/GameplayTagCustomization.cpp
+++ b/Engine/Plugins/Editor/GameplayTagsEditor/Source/GameplayTagsEditor/Private/GameplayTagCustomization.cpp
@@ -54,6 +54,7 @@ void FGameplayTagCustomization::CustomizeHeader(TSharedRef<IPropertyHandle> InSt
 	bool bHasValidCategories = CurrentTag == FGameplayTag::EmptyTag || MetaDataCategories.IsEmpty() || CurrentTag.MatchesAny(MetaDataCategories);
 
 	HeaderRow
+	.FilterString(FText::AsCultureInvariant(CurrentTag.ToString()))
 	.NameContent()
 	[
 		SNew(SHorizontalBox)
