Splitting Tag ini into multiple based on parent level tag.

#Synced-CL 52455004

diff --git a/Engine/Plugins/Editor/GameplayTagsEditor/Source/GameplayTagsEditor/Private/SplitGameplayTagsCommandlet.cpp b/Engine/Plugins/Editor/GameplayTagsEditor/Source/GameplayTagsEditor/Private/SplitGameplayTagsCommandlet.cpp
new file mode 100644
--- /dev/null
+++ b/Engine/Plugins/Editor/GameplayTagsEditor/Source/GameplayTagsEditor/Private/SplitGameplayTagsCommandlet.cpp
@@ -0,0 +1,274 @@
+// Copyright Epic Games, Inc. All Rights Reserved.
+
+#include "SplitGameplayTagsCommandlet.h"
+
+#include "HAL/FileManager.h"
+#include "ISourceControlModule.h"
+#include "Misc/FileHelper.h"
+#include "Misc/Paths.h"
+#include "SourceControlHelpers.h"
+
+DEFINE_LOG_CATEGORY_STATIC(LogSplitGameplayTags, Log, All);
+
+USplitGameplayTagsCommandlet::USplitGameplayTagsCommandlet()
+{
+	IsEditor = true;
+}
+
+FString USplitGameplayTagsCommandlet::ExtractTagName(const FString& Line)
+{
+	// Format: +GameplayTagList=(Tag="Foo.Bar.Baz",DevComment="...")
+	// Locate the opening quote after the 'Tag=' key.
+	static const TCHAR* TagKey = TEXT("Tag=\"");
+	static const int32  TagKeyLen = 5; // T a g = "
+
+	const int32 KeyStart = Line.Find(TagKey, ESearchCase::CaseSensitive);
+	if (KeyStart == INDEX_NONE)
+	{
+		return FString();
+	}
+
+	const int32 ValueStart = KeyStart + TagKeyLen;
+	const int32 ValueEnd   = Line.Find(TEXT("\""), ESearchCase::CaseSensitive, ESearchDir::FromStart, ValueStart);
+	if (ValueEnd == INDEX_NONE)
+	{
+		return FString();
+	}
+
+	return Line.Mid(ValueStart, ValueEnd - ValueStart);
+}
+
+FString USplitGameplayTagsCommandlet::GetTopLevelPrefix(const FString& TagName)
+{
+	int32 DotIndex = INDEX_NONE;
+	TagName.FindChar(TEXT('.'), DotIndex);
+	return (DotIndex != INDEX_NONE) ? TagName.Left(DotIndex) : TagName;
+}
+
+FString USplitGameplayTagsCommandlet::ToTagListLine(const FString& DefaultTagLine)
+{
+	// DefaultGameplayTags.ini: "+GameplayTagList=(...)"
+	// Config/Tags split files:  "GameplayTagList=(...)"  (no leading '+')
+	return DefaultTagLine.StartsWith(TEXT("+")) ? DefaultTagLine.Mid(1) : DefaultTagLine;
+}
+
+int32 USplitGameplayTagsCommandlet::Main(const FString& Params)
+{
+	TArray<FString> Tokens, Switches;
+	TMap<FString, FString> ParamMap;
+	ParseCommandLine(*Params, Tokens, Switches, ParamMap);
+
+	const bool  bDryRun = Switches.Contains(TEXT("DryRun"));
+	const bool  bNoSort = Switches.Contains(TEXT("NoSort"));
+	const int32 MinTags = ParamMap.Contains(TEXT("MinTags"))
+		? FCString::Atoi(*ParamMap[TEXT("MinTags")]) : 50;
+
+	FString SourceIniPath = ParamMap.Contains(TEXT("SourceIni"))
+		? ParamMap[TEXT("SourceIni")]
+		: FPaths::ProjectConfigDir() / TEXT("DefaultGameplayTags.ini");
+	FPaths::NormalizeFilename(SourceIniPath);
+
+	FString OutputDir = ParamMap.Contains(TEXT("OutputDir"))
+		? ParamMap[TEXT("OutputDir")]
+		: FPaths::ProjectConfigDir() / TEXT("Tags");
+	FPaths::NormalizeDirectoryName(OutputDir);
+
+	// -----------------------------------------------------------------------
+	// 1. Read source file
+	// -----------------------------------------------------------------------
+	TArray<FString> SourceLines;
+	if (!FFileHelper::LoadFileToStringArray(SourceLines, *SourceIniPath))
+	{
+		UE_LOG(LogSplitGameplayTags, Error, TEXT("Failed to read source ini: %s"), *SourceIniPath);
+		return 1;
+	}
+	UE_LOG(LogSplitGameplayTags, Display, TEXT("Read %d lines from %s"), SourceLines.Num(), *SourceIniPath);
+
+	// -----------------------------------------------------------------------
+	// 2. Partition lines into header content and tag entries.
+	//    Tag lines are grouped by case-insensitive top-level prefix.
+	//    The first-seen casing of each prefix is used for the output filename.
+	// -----------------------------------------------------------------------
+	TArray<FString>                HeaderLines;
+	TMap<FString, TArray<FString>> TagsByNormPrefix;  // lowercase-prefix -> original lines
+	TMap<FString, FString>         PrefixDisplayName; // lowercase-prefix -> first-seen casing
+	int32 MalformedCount = 0;
+
+	for (const FString& Line : SourceLines)
+	{
+		const FString Trimmed = Line.TrimStart();
+		if (!Trimmed.StartsWith(TEXT("+GameplayTagList=")))
+		{
+			HeaderLines.Add(Line);
+			continue;
+		}
+
+		const FString TagName = ExtractTagName(Trimmed);
+		if (TagName.IsEmpty())
+		{
+			UE_LOG(LogSplitGameplayTags, Warning,
+				TEXT("Could not parse tag name from line (retaining in source): %s"), *Line);
+			HeaderLines.Add(Line);
+			++MalformedCount;
+			continue;
+		}
+
+		const FString RawPrefix  = GetTopLevelPrefix(TagName);
+		const FString NormPrefix = RawPrefix.ToLower();
+
+		TagsByNormPrefix.FindOrAdd(NormPrefix).Add(Line);
+		if (!PrefixDisplayName.Contains(NormPrefix))
+		{
+			PrefixDisplayName.Add(NormPrefix, RawPrefix);
+		}
+	}
+
+	const int32 TotalTagLines = SourceLines.Num() - HeaderLines.Num() - MalformedCount;
+	UE_LOG(LogSplitGameplayTags, Display,
+		TEXT("Parsed %d tag lines across %d unique prefixes (%d malformed, retained in source)"),
+		TotalTagLines, TagsByNormPrefix.Num(), MalformedCount);
+
+	// -----------------------------------------------------------------------
+	// 3. Assign prefix groups to output files.
+	//    Groups >= MinTags get a dedicated GameplayTags_<Prefix>.ini.
+	//    Smaller groups fold into GameplayTags_Misc.ini.
+	// -----------------------------------------------------------------------
+	TMap<FString, TArray<FString>> FileBuckets; // file-stem -> lines to write
+	int32 MiscPrefixCount = 0;
+
+	for (auto& [NormPrefix, Lines] : TagsByNormPrefix)
+	{
+		if (Lines.Num() >= MinTags)
+		{
+			const FString Stem = FString::Printf(TEXT("GameplayTags_%s"), *PrefixDisplayName[NormPrefix]);
+			FileBuckets.FindOrAdd(Stem).Append(Lines);
+		}
+		else
+		{
+			FileBuckets.FindOrAdd(TEXT("GameplayTags_Misc")).Append(Lines);
+			++MiscPrefixCount;
+		}
+	}
+
+	// -----------------------------------------------------------------------
+	// 4. Sort tag lines within each bucket for deterministic P4 diffs.
+	//    Alphabetically sorted files produce cleaner future merges since
+	//    independently-added tags won't cluster at the same insertion point.
+	// -----------------------------------------------------------------------
+	if (!bNoSort)
+	{
+		for (auto& [Stem, Lines] : FileBuckets)
+		{
+			Lines.Sort();
+		}
+	}
+
+	// -----------------------------------------------------------------------
+	// 5. Report planned output
+	// -----------------------------------------------------------------------
+	TArray<FString> SortedStems;
+	FileBuckets.GetKeys(SortedStems);
+	SortedStems.Sort();
+
+	int32 TotalTagsToWrite = 0;
+	for (const FString& Stem : SortedStems)
+	{
+		const int32 Count = FileBuckets[Stem].Num();
+		TotalTagsToWrite += Count;
+
+		if (Stem == TEXT("GameplayTags_Misc"))
+		{
+			UE_LOG(LogSplitGameplayTags, Display,
+				TEXT("  -> %s.ini  (%d tags, %d prefix groups merged)"), *Stem, Count, MiscPrefixCount);
+		}
+		else
+		{
+			UE_LOG(LogSplitGameplayTags, Display, TEXT("  -> %s.ini  (%d tags)"), *Stem, Count);
+		}
+	}
+
+	UE_LOG(LogSplitGameplayTags, Display,
+		TEXT("Summary: %d tags -> %d output files | %d header lines retained | output dir: %s"),
+		TotalTagsToWrite, FileBuckets.Num(), HeaderLines.Num(), *OutputDir);
+
+	if (bDryRun)
+	{
+		UE_LOG(LogSplitGameplayTags, Display, TEXT("-DryRun specified, no files written."));
+		return 0;
+	}
+
+	// -----------------------------------------------------------------------
+	// 6. Write split files to Config/Tags/
+	// -----------------------------------------------------------------------
+	if (!IFileManager::Get().DirectoryExists(*OutputDir))
+	{
+		if (!IFileManager::Get().MakeDirectory(*OutputDir, /*Tree=*/true))
+		{
+			UE_LOG(LogSplitGameplayTags, Error, TEXT("Failed to create output directory: %s"), *OutputDir);
+			return 1;
+		}
+	}
+
+	const bool bSCEnabled = ISourceControlModule::Get().IsEnabled();
+	static const FString SplitSectionHeader = TEXT("[/Script/GameplayTags.GameplayTagsList]");
+
+	for (const FString& Stem : SortedStems)
+	{
+		const TArray<FString>& BucketLines = FileBuckets[Stem];
+		const FString OutPath = OutputDir / Stem + TEXT(".ini");
+
+		FString Content;
+		Content.Reserve((BucketLines.Num() + 1) * 80);
+		Content += SplitSectionHeader;
+		Content += LINE_TERMINATOR;
+
+		for (const FString& RawLine : BucketLines)
+		{
+			Content += ToTagListLine(RawLine);
+			Content += LINE_TERMINATOR;
+		}
+
+		if (!FFileHelper::SaveStringToFile(Content, *OutPath, FFileHelper::EEncodingOptions::ForceUTF8WithoutBOM))
+		{
+			UE_LOG(LogSplitGameplayTags, Error, TEXT("Failed to write: %s"), *OutPath);
+			return 1;
+		}
+
+		if (bSCEnabled)
+		{
+			USourceControlHelpers::CheckOutOrAddFile(OutPath);
+		}
+
+		UE_LOG(LogSplitGameplayTags, Display, TEXT("Wrote %s  (%d tags)"), *OutPath, BucketLines.Num());
+	}
+
+	// -----------------------------------------------------------------------
+	// 7. Rewrite DefaultGameplayTags.ini with all +GameplayTagList entries removed.
+	//    All other content (settings, CategoryRemapping, GameplayTagTableList,
+	//    GameplayTagRedirects) is preserved exactly as-is.
+	// -----------------------------------------------------------------------
+	if (bSCEnabled)
+	{
+		USourceControlHelpers::CheckOutOrAddFile(SourceIniPath);
+	}
+
+	FString NewSourceContent;
+	NewSourceContent.Reserve(HeaderLines.Num() * 80);
+	for (const FString& Line : HeaderLines)
+	{
+		NewSourceContent += Line;
+		NewSourceContent += LINE_TERMINATOR;
+	}
+
+	if (!FFileHelper::SaveStringToFile(NewSourceContent, *SourceIniPath, FFileHelper::EEncodingOptions::ForceUTF8WithoutBOM))
+	{
+		UE_LOG(LogSplitGameplayTags, Error, TEXT("Failed to rewrite source ini: %s"), *SourceIniPath);
+		return 1;
+	}
+
+	UE_LOG(LogSplitGameplayTags, Display,
+		TEXT("Rewrote %s (%d lines retained, %d tag lines removed)"),
+		*SourceIniPath, HeaderLines.Num(), TotalTagsToWrite);
+
+	return 0;
+}
diff --git a/Engine/Plugins/Editor/GameplayTagsEditor/Source/GameplayTagsEditor/Private/SplitGameplayTagsCommandlet.h b/Engine/Plugins/Editor/GameplayTagsEditor/Source/GameplayTagsEditor/Private/SplitGameplayTagsCommandlet.h
new file mode 100644
--- /dev/null
+++ b/Engine/Plugins/Editor/GameplayTagsEditor/Source/GameplayTagsEditor/Private/SplitGameplayTagsCommandlet.h
@@ -0,0 +1,64 @@
+// Copyright Epic Games, Inc. All Rights Reserved.
+
+#pragma once
+
+#include "CoreMinimal.h"
+#include "Commandlets/Commandlet.h"
+#include "SplitGameplayTagsCommandlet.generated.h"
+
+/**
+ * One-shot migration commandlet that splits the monolithic DefaultGameplayTags.ini
+ * into per-prefix tag list files under Config/Tags/, reducing Perforce merge
+ * conflict surface area.
+ *
+ * Tags are grouped by their top-level namespace prefix (the first node before the
+ * first dot). Prefixes with fewer than MinTags entries are folded into a catch-all
+ * GameplayTags_Misc.ini rather than creating dozens of near-empty files.
+ *
+ * Output files use the [/Script/GameplayTags.GameplayTagsList] section and the
+ * GameplayTagList=(Tag=...) format (no leading '+') expected by UGameplayTagsList,
+ * matching the format of existing files already in Config/Tags/.
+ *
+ * All non-tag content in DefaultGameplayTags.ini is preserved in place:
+ * settings, CategoryRemapping, GameplayTagTableList, and GameplayTagRedirects
+ * entries are not touched.
+ *
+ * Usage:
+ *   UnrealEditor.exe <Project> -run=SplitGameplayTags [options]
+ *
+ *   -DryRun          Report what would be done without writing any files.
+ *   -NoSort          Preserve original tag order within output files.
+ *                    Default: sort alphabetically for deterministic P4 diffs.
+ *   -MinTags=N       Minimum tag count for a prefix to get its own file.
+ *                    Default: 50. Lower values produce more, smaller files.
+ *   -OutputDir=<p>   Target directory for split files (default: Config/Tags/).
+ *   -SourceIni=<p>   Source ini to split (default: Config/DefaultGameplayTags.ini).
+ *
+ * After running: p4 add Config/Tags/GameplayTags_*.ini, then submit both the
+ * new files and the modified DefaultGameplayTags.ini in a single changelist.
+ */
+UCLASS()
+class USplitGameplayTagsCommandlet : public UCommandlet
+{
+	GENERATED_BODY()
+
+public:
+	USplitGameplayTagsCommandlet();
+	virtual int32 Main(const FString& Params) override;
+
+private:
+	/**
+	 * Extracts the tag name value from a raw +GameplayTagList=(Tag="...",DevComment="...") line.
+	 * Returns an empty string if the line cannot be parsed.
+	 */
+	static FString ExtractTagName(const FString& Line);
+
+	/** Returns the first dot-separated node of TagName (i.e. the top-level prefix). */
+	static FString GetTopLevelPrefix(const FString& TagName);
+
+	/**
+	 * Converts a DefaultGameplayTags.ini tag line to the format used in Config/Tags split files.
+	 * Strips the leading '+' from "+GameplayTagList=(...)" -> "GameplayTagList=(...)".
+	 */
+	static FString ToTagListLine(const FString& DefaultTagLine);
+};
